Renable shadow variable warning pragma around glsl_parser_extras

diff --git a/Engine/Source/Developer/Apple/MetalShaderFormat/Private/MetalBackend.cpp b/Engine/Source/Developer/Apple/MetalShaderFormat/Private/MetalBackend.cpp
--- a/Engine/Source/Developer/Apple/MetalShaderFormat/Private/MetalBackend.cpp
+++ b/Engine/Source/Developer/Apple/MetalShaderFormat/Private/MetalBackend.cpp
@@ -8,9 +8,9 @@
 #include "MetalBackend.h"
 #include "compiler.h"
 
-//PRAGMA_DISABLE_SHADOW_VARIABLE_WARNINGS
+PRAGMA_DISABLE_SHADOW_VARIABLE_WARNINGS
 #include "glsl_parser_extras.h"
-//PRAGMA_POP
+PRAGMA_POP
 
 #include "hash_table.h"
 #include "ir_rvalue_visitor.h"
