fixed CIS - metal compile

diff --git a/Engine/Source/Developer/Apple/MetalShaderFormat/Private/MetalShaderCompiler.cpp b/Engine/Source/Developer/Apple/MetalShaderFormat/Private/MetalShaderCompiler.cpp
--- a/Engine/Source/Developer/Apple/MetalShaderFormat/Private/MetalShaderCompiler.cpp
+++ b/Engine/Source/Developer/Apple/MetalShaderFormat/Private/MetalShaderCompiler.cpp
@@ -410,7 +410,7 @@ static void BuildMetalShaderOutput(
 						bCompileAtRuntime = false;
 						
 						// Write out the header and compiled shader code
-						FMemoryWriter Ar(ShaderOutput.Code.ShaderCode.GetWriteAccess(), true);
+						FMemoryWriter Ar(ShaderOutput.ShaderCode.GetWriteAccess(), true);
 						uint8 PrecompiledFlag = 1;
 						Ar << PrecompiledFlag;
 						Ar << Header;
