UE4-Main - hlslcc - Fix for Metal backend not propagating integer types on textures

diff --git a/Engine/Source/Developer/Apple/MetalShaderFormat/Private/MetalBackend.cpp b/Engine/Source/Developer/Apple/MetalShaderFormat/Private/MetalBackend.cpp
--- a/Engine/Source/Developer/Apple/MetalShaderFormat/Private/MetalBackend.cpp
+++ b/Engine/Source/Developer/Apple/MetalShaderFormat/Private/MetalBackend.cpp
@@ -776,9 +776,27 @@ protected:
 							}
 
 							print_type_pre(PtrType);
+							const char* InnerType = "float";
+							if (PtrType->inner_type)
+							{
+								switch (PtrType->inner_type->base_type)
+								{
+									case GLSL_TYPE_HALF:
+										InnerType = "half";
+										break;
+									case GLSL_TYPE_INT:
+										InnerType = "int";
+										break;
+									case GLSL_TYPE_UINT:
+										InnerType = "uint";
+										break;
+									default:
+										break;
+								}
+							}
 							ralloc_asprintf_append(
 								buffer,
-								"<%s> %s", (PtrType->inner_type && PtrType->inner_type->base_type == GLSL_TYPE_HALF) ? "half" : "float", unique_name(var));
+								"<%s> %s", InnerType, unique_name(var));
 							print_type_post(PtrType);
 							ralloc_asprintf_append(
 								buffer,
