UE-Main - Fix depth textures

diff --git a/Engine/Source/Developer/Apple/MetalShaderFormat/Private/MetalBackend.cpp b/Engine/Source/Developer/Apple/MetalShaderFormat/Private/MetalBackend.cpp
--- a/Engine/Source/Developer/Apple/MetalShaderFormat/Private/MetalBackend.cpp
+++ b/Engine/Source/Developer/Apple/MetalShaderFormat/Private/MetalBackend.cpp
@@ -779,8 +779,14 @@ protected:
 							const char* InnerType = "float";
 							if (PtrType->inner_type)
 							{
-								switch (PtrType->inner_type->base_type)
+								if (PtrType->base_type == GLSL_TYPE_SAMPLER && PtrType->sampler_shadow)
 								{
+									//#todo-rco: Currently force to float...
+								}
+								else
+								{
+									switch (PtrType->inner_type->base_type)
+									{
 									case GLSL_TYPE_HALF:
 										InnerType = "half";
 										break;
@@ -792,6 +798,7 @@ protected:
 										break;
 									default:
 										break;
+									}
 								}
 							}
 							ralloc_asprintf_append(
