Merging CL#2822520

#UE4 Non-unity fix for MetalUtils.h

diff --git a/Engine/Source/Developer/Apple/MetalShaderFormat/Private/MetalUtils.h b/Engine/Source/Developer/Apple/MetalShaderFormat/Private/MetalUtils.h
--- a/Engine/Source/Developer/Apple/MetalShaderFormat/Private/MetalUtils.h
+++ b/Engine/Source/Developer/Apple/MetalShaderFormat/Private/MetalUtils.h
@@ -3,6 +3,7 @@
 #pragma once
 
 #include "hlslcc.h"
+#include "CustomStdAllocator.h"
 
 inline FCustomStdString FixVecPrefix(FCustomStdString Type)
 {
