UE4 - Fix for mac non unity

diff --git a/Engine/Source/Developer/Apple/MetalShaderFormat/Private/MetalBackend.cpp b/Engine/Source/Developer/Apple/MetalShaderFormat/Private/MetalBackend.cpp
--- a/Engine/Source/Developer/Apple/MetalShaderFormat/Private/MetalBackend.cpp
+++ b/Engine/Source/Developer/Apple/MetalShaderFormat/Private/MetalBackend.cpp
@@ -6,6 +6,7 @@
 #include "hlslcc.h"
 #include "hlslcc_private.h"
 #include "MetalBackend.h"
+#include "MetalUtils.h"
 #include "compiler.h"
 
 PRAGMA_DISABLE_SHADOW_VARIABLE_WARNINGS
diff --git a/Engine/Source/Developer/Apple/MetalShaderFormat/Private/MetalBackend.h b/Engine/Source/Developer/Apple/MetalShaderFormat/Private/MetalBackend.h
--- a/Engine/Source/Developer/Apple/MetalShaderFormat/Private/MetalBackend.h
+++ b/Engine/Source/Developer/Apple/MetalShaderFormat/Private/MetalBackend.h
@@ -3,6 +3,8 @@
 #pragma once
 
 #include "hlslcc.h"
+#include "ir.h"
+#include "PackUniformBuffers.h"
 #include "LanguageSpec.h"
 
 
diff --git a/Engine/Source/Developer/Apple/MetalShaderFormat/Private/MetalUtils.h b/Engine/Source/Developer/Apple/MetalShaderFormat/Private/MetalUtils.h
--- a/Engine/Source/Developer/Apple/MetalShaderFormat/Private/MetalUtils.h
+++ b/Engine/Source/Developer/Apple/MetalShaderFormat/Private/MetalUtils.h
@@ -3,6 +3,8 @@
 #pragma once
 
 #include "hlslcc.h"
+#include "ir.h"
+#include "PackUniformBuffers.h"
 #include "CustomStdAllocator.h"
 
 inline FCustomStdString FixVecPrefix(FCustomStdString Type)
