Fixed Mac ShaderCompileWorker crash

diff --git a/Engine/Source/Developer/Apple/MetalShaderFormat/Private/MetalDerivedData.cpp b/Engine/Source/Developer/Apple/MetalShaderFormat/Private/MetalDerivedData.cpp
--- a/Engine/Source/Developer/Apple/MetalShaderFormat/Private/MetalDerivedData.cpp
+++ b/Engine/Source/Developer/Apple/MetalShaderFormat/Private/MetalDerivedData.cpp
@@ -134,7 +134,7 @@ FString FMetalShaderBytecodeCooker::GetPluginSpecificCacheKeySuffix() const
 {
 	FString CompilerVersion = GetMetalCompilerVersion(MetalShaderFormatToLegacyShaderPlatform(Job.ShaderFormat));
 		
-	FString VersionedName = FString::Printf(TEXT("%s%u%u%s%s%s%s%s%s%s"), *Job.ShaderFormat.GetPlainNameString(), Job.SourceCRCLen, Job.SourceCRC, *Job.Hash.ToString(), *Job.CompilerVersion, *Job.MinOSVersion, *Job.DebugInfo, *Job.MathMode, *Job.Standard, Job.bRetainObjectFile ? ("+Object") : (""));
+	FString VersionedName = FString::Printf(TEXT("%s%u%u%s%s%s%s%s%s%s"), *Job.ShaderFormat.GetPlainNameString(), Job.SourceCRCLen, Job.SourceCRC, *Job.Hash.ToString(), *Job.CompilerVersion, *Job.MinOSVersion, *Job.DebugInfo, *Job.MathMode, *Job.Standard, Job.bRetainObjectFile ? TEXT("+Object") : TEXT(""));
 	// get rid of some not so filename-friendly characters ('=',' ' -> '_')
     VersionedName = VersionedName.Replace(TEXT("="), TEXT("_")).Replace(TEXT(" "), TEXT("_"));
 
