Add a hash of the Metal compiler toolchain path into the Metal bytecode key so that when recompiling the metal_stdlib on a Mac where Xcode is installed to a different place we don't just fail.

diff --git a/Engine/Source/Developer/Apple/MetalShaderFormat/Private/MetalDerivedData.cpp b/Engine/Source/Developer/Apple/MetalShaderFormat/Private/MetalDerivedData.cpp
--- a/Engine/Source/Developer/Apple/MetalShaderFormat/Private/MetalDerivedData.cpp
+++ b/Engine/Source/Developer/Apple/MetalShaderFormat/Private/MetalDerivedData.cpp
@@ -133,8 +133,9 @@ const TCHAR* FMetalShaderBytecodeCooker::GetVersionString() const
 FString FMetalShaderBytecodeCooker::GetPluginSpecificCacheKeySuffix() const
 {
 	FString CompilerVersion = GetMetalCompilerVersion(MetalShaderFormatToLegacyShaderPlatform(Job.ShaderFormat));
-		
-	FString VersionedName = FString::Printf(TEXT("%s%u%u%s%s%s%s%s%s%s"), *Job.ShaderFormat.GetPlainNameString(), Job.SourceCRCLen, Job.SourceCRC, *Job.Hash.ToString(), *Job.CompilerVersion, *Job.MinOSVersion, *Job.DebugInfo, *Job.MathMode, *Job.Standard, Job.bRetainObjectFile ? TEXT("+Object") : TEXT(""));
+	FString CompilerPath = GetMetalToolsPath(MetalShaderFormatToLegacyShaderPlatform(Job.ShaderFormat));
+	
+	FString VersionedName = FString::Printf(TEXT("%s%u%u%s%s%s%s%s%s%s%d"), *Job.ShaderFormat.GetPlainNameString(), Job.SourceCRCLen, Job.SourceCRC, *Job.Hash.ToString(), *Job.CompilerVersion, *Job.MinOSVersion, *Job.DebugInfo, *Job.MathMode, *Job.Standard, Job.bRetainObjectFile ? TEXT("+Object") : TEXT(""), GetTypeHash(CompilerPath));
 	// get rid of some not so filename-friendly characters ('=',' ' -> '_')
     VersionedName = VersionedName.Replace(TEXT("="), TEXT("_")).Replace(TEXT(" "), TEXT("_"));
 
