Merging //UE4/Release-4.18 @ 3660543 to Release-Staging-4.18 (//UE4/Release-Staging-4.18)

diff --git a/Engine/Source/Developer/Apple/MetalShaderFormat/Private/MetalBackend.cpp b/Engine/Source/Developer/Apple/MetalShaderFormat/Private/MetalBackend.cpp
--- a/Engine/Source/Developer/Apple/MetalShaderFormat/Private/MetalBackend.cpp
+++ b/Engine/Source/Developer/Apple/MetalShaderFormat/Private/MetalBackend.cpp
@@ -1542,6 +1542,16 @@ protected:
             expr->operands[0]->accept(this);
             ralloc_asprintf_append(buffer, ")");
         }
+		else if (Backend && Backend->Version >= 2 && numOps == 2 && op == ir_binop_mul && expr->operands[0]->type == expr->operands[1]->type && expr->operands[0]->type->is_float())
+		{
+			ralloc_asprintf_append(buffer, "fma(");
+			expr->operands[0]->accept(this);
+			ralloc_asprintf_append(buffer, ",");
+			expr->operands[1]->accept(this);
+			ralloc_asprintf_append(buffer, ",");
+			print_type_full(expr->operands[0]->type);
+			ralloc_asprintf_append(buffer, "(0))");
+		}
 		else if (numOps == 2 && (op == ir_binop_add || op == ir_binop_sub || op == ir_binop_mul || op == ir_binop_div))
 		{
 			bool bHandleFloatHalfConflict = false;
diff --git a/Engine/Source/Developer/Apple/MetalShaderFormat/Private/MetalShaderCompiler.cpp b/Engine/Source/Developer/Apple/MetalShaderFormat/Private/MetalShaderCompiler.cpp
--- a/Engine/Source/Developer/Apple/MetalShaderFormat/Private/MetalShaderCompiler.cpp
+++ b/Engine/Source/Developer/Apple/MetalShaderFormat/Private/MetalShaderCompiler.cpp
@@ -1689,13 +1689,6 @@ void CompileShader_Metal(const FShaderCompilerInput& _Input,FShaderCompilerOutpu
 	bool const bDirectCompile = FParse::Param(FCommandLine::Get(), TEXT("directcompile"));
 	if (!Input.bSkipPreprocessedCache && !bDirectCompile)
 	{
-		FString const* UsingWPO = Input.Environment.GetDefinitions().Find(TEXT("USES_WORLD_POSITION_OFFSET"));
-		if (UsingWPO && FString("1") == *UsingWPO)
-		{
-			// FMAs are deoptimised when fast-math is enabled and arguments are literals :(
-			Input.Environment.CompilerFlags.Add(CFLAG_NoFastMath);
-		}
-		
 		FString const* UsingTessellationDefine = Input.Environment.GetDefinitions().Find(TEXT("USING_TESSELLATION"));
 		bool bUsingTessellation = (UsingTessellationDefine != nullptr && FString("1") == *UsingTessellationDefine);
 		if (bUsingTessellation && (Input.Target.Frequency == SF_Vertex))
diff --git a/Engine/Source/Developer/Apple/MetalShaderFormat/Private/MetalShaderFormat.cpp b/Engine/Source/Developer/Apple/MetalShaderFormat/Private/MetalShaderFormat.cpp
--- a/Engine/Source/Developer/Apple/MetalShaderFormat/Private/MetalShaderFormat.cpp
+++ b/Engine/Source/Developer/Apple/MetalShaderFormat/Private/MetalShaderFormat.cpp
@@ -107,7 +107,7 @@ class FMetalShaderFormat : public IShaderFormat
 public:
 	enum
 	{
-		HEADER_VERSION = 42,
+		HEADER_VERSION = 43,
 	};
 	
 	struct FVersion
