Fix static analysis warning after merge from Release-Staging-4.18.

diff --git a/Engine/Source/Developer/Apple/MetalShaderFormat/Private/MetalBackend.cpp b/Engine/Source/Developer/Apple/MetalShaderFormat/Private/MetalBackend.cpp
--- a/Engine/Source/Developer/Apple/MetalShaderFormat/Private/MetalBackend.cpp
+++ b/Engine/Source/Developer/Apple/MetalShaderFormat/Private/MetalBackend.cpp
@@ -1625,7 +1625,7 @@ protected:
 		}
 		else if ((op == ir_ternop_fma || op == ir_ternop_clamp || op == ir_unop_sqrt || op == ir_unop_rsq || op == ir_unop_saturate) && expr->type->base_type == GLSL_TYPE_FLOAT)
 		{
-			if (!Backend->bAllowFastIntriniscs && op != ir_ternop_fma)
+			if (Backend && !Backend->bAllowFastIntriniscs && op != ir_ternop_fma)
 			{
 				ralloc_asprintf_append(buffer, "precise::");
 			}
@@ -1646,7 +1646,7 @@ protected:
 			{
 				OpString = (OpString + 1);
 			}
-			else if(!Backend->bAllowFastIntriniscs && expr->type->base_type == GLSL_TYPE_FLOAT)
+			else if(Backend && !Backend->bAllowFastIntriniscs && expr->type->base_type == GLSL_TYPE_FLOAT)
 			{
 				ralloc_asprintf_append(buffer, "precise::");
 			}
