Fixed a mistake I let slip through in 3934147 that will potentially cause slight miscompilation of SP_METAL.

diff --git a/Engine/Source/Developer/Apple/MetalShaderFormat/Private/MetalShaderCompiler.cpp b/Engine/Source/Developer/Apple/MetalShaderFormat/Private/MetalShaderCompiler.cpp
--- a/Engine/Source/Developer/Apple/MetalShaderFormat/Private/MetalShaderCompiler.cpp
+++ b/Engine/Source/Developer/Apple/MetalShaderFormat/Private/MetalShaderCompiler.cpp
@@ -978,10 +978,10 @@ void BuildMetalShaderOutput(
 	const ANSICHAR* SideTableString = FCStringAnsi::Strstr(USFSource, "@SideTable: ");
 	
 	EShaderFrequency Frequency = (EShaderFrequency)ShaderOutput.Target.Frequency;
-	const bool bIsMobile = (ShaderInput.Target.Platform == SP_METAL_MRT);
+	const bool bIsMobile = (ShaderInput.Target.Platform == SP_METAL || ShaderInput.Target.Platform == SP_METAL_MRT);
 	bool bNoFastMath = ShaderInput.Environment.CompilerFlags.Contains(CFLAG_NoFastMath);
 	FString const* UsingWPO = ShaderInput.Environment.GetDefinitions().Find(TEXT("USES_WORLD_POSITION_OFFSET"));
-	if (UsingWPO && FString("1") == *UsingWPO && bIsMobile && Frequency == SF_Vertex)
+	if (UsingWPO && FString("1") == *UsingWPO && (ShaderInput.Target.Platform == SP_METAL_MRT) && Frequency == SF_Vertex)
 	{
 		// WPO requires that we make all multiply/sincos instructions invariant :(
 		bNoFastMath = true;
