Update for new Xcode metal shader compiler requiring different arguments. Make some slight changes to make some behaviour and error reports clearer.

diff --git a/Engine/Source/Developer/Apple/MetalShaderFormat/Private/MetalDerivedData.cpp b/Engine/Source/Developer/Apple/MetalShaderFormat/Private/MetalDerivedData.cpp
--- a/Engine/Source/Developer/Apple/MetalShaderFormat/Private/MetalDerivedData.cpp
+++ b/Engine/Source/Developer/Apple/MetalShaderFormat/Private/MetalDerivedData.cpp
@@ -210,11 +210,11 @@ bool FMetalShaderBytecodeCooker::Build(TArray<uint8>& OutData)
 		if (bUseSharedPCH)
         {
             CopyLocalFileToRemote(Job.InputPCHFile, RemoteInputPCHFile);
-			MetalParams = FString::Printf(TEXT("-include-pch %s %s %s %s -Wno-null-character -fbracket-depth=1024 %s %s %s %s -o %s"), *RemoteInputPCHFile, *Job.MinOSVersion, *Job.DebugInfo, *Job.MathMode, *Job.Standard, *Job.Defines, *IncludeArgs, *RemoteInputFile, *RemoteObjFile);
+			MetalParams = FString::Printf(TEXT("-include-pch %s %s %s %s -c -Wno-null-character -fbracket-depth=1024 %s %s %s %s -o %s"), *RemoteInputPCHFile, *Job.MinOSVersion, *Job.DebugInfo, *Job.MathMode, *Job.Standard, *Job.Defines, *IncludeArgs, *RemoteInputFile, *RemoteObjFile);
         }
         else
         {
-            MetalParams = FString::Printf(TEXT("%s %s %s -Wno-null-character -fbracket-depth=1024 %s %s %s %s -o %s"), *Job.MinOSVersion, *Job.DebugInfo, *Job.MathMode, *Job.Standard, *Job.Defines, *IncludeArgs, *RemoteInputFile, *RemoteObjFile);
+            MetalParams = FString::Printf(TEXT("%s %s %s -c -Wno-null-character -fbracket-depth=1024 %s %s %s %s -o %s"), *Job.MinOSVersion, *Job.DebugInfo, *Job.MathMode, *Job.Standard, *Job.Defines, *IncludeArgs, *RemoteInputFile, *RemoteObjFile);
         }
 	}
 
@@ -267,11 +267,11 @@ bool FMetalShaderBytecodeCooker::Build(TArray<uint8>& OutData)
 	{
 		if (Job.bCompileAsPCH)
 		{
-			Job.Message = FString::Printf(TEXT("Metal Shared PCH generation failed %s: %s."), CompileType, *Job.Errors);
+			Job.Message = FString::Printf(TEXT("Metal Shared PCH generation failed %s to generate %s: %s."), CompileType, *RemoteOutputFilename, *Job.Errors);
 		}
 		else
 		{
-			Job.Message = FString::Printf(TEXT("Failed to compile to bytecode %s, code: %d, output: %s %s"), CompileType, Job.ReturnCode, *Job.Results, *Job.Errors);
+			Job.Message = FString::Printf(TEXT("Failed to compile %s to bytecode %s, code: %d, output: %s %s"), CompileType, *RemoteOutputFilename, Job.ReturnCode, *Job.Results, *Job.Errors);
 		}
 	}
 
diff --git a/Engine/Source/Developer/Apple/MetalShaderFormat/Private/MetalShaderCompiler.cpp b/Engine/Source/Developer/Apple/MetalShaderFormat/Private/MetalShaderCompiler.cpp
--- a/Engine/Source/Developer/Apple/MetalShaderFormat/Private/MetalShaderCompiler.cpp
+++ b/Engine/Source/Developer/Apple/MetalShaderFormat/Private/MetalShaderCompiler.cpp
@@ -475,7 +475,7 @@ FString GetMetalBinaryPath(uint32 ShaderPlatform)
 		FString XcodePath = GetXcodePath();
 		if (XcodePath.Len() > 0)
 		{
-			FString MetalToolsPath = FString::Printf(TEXT("%s/Toolchains/XcodeDefault.xctoolchain/usr/bin"), *XcodePath);
+			FString MetalToolsPath = FString::Printf(TEXT("%s/Toolchains/XcodeDefault.xctoolchain/usr/metal/macos/bin"), *XcodePath);
 			FString MetalPath = MetalToolsPath + TEXT("/metal");
 			if (!RemoteFileExists(MetalPath))
 			{
@@ -497,15 +497,7 @@ FString GetMetalBinaryPath(uint32 ShaderPlatform)
 				
 				GMetalCompilerVers[bIsMobile] = GetMetalCompilerVers(MetalPath);
 				
-				FString MetalLibraryPath;
-				if (bIsMobile)
-				{
-					MetalLibraryPath = FString::Printf(TEXT("%s/Platforms/iPhoneOS.platform/usr/lib"), *XcodePath);
-				}
-				else
-				{
-					MetalLibraryPath = FString::Printf(TEXT("%s/Platforms/MacOSX.platform/usr/lib"), *XcodePath);
-				}
+				FString MetalLibraryPath = FString::Printf(TEXT("%s/../lib"), *MetalToolsPath);
 				
 				FString MetalStdLibPath = GetMetalStdLibPath(MetalLibraryPath);
 				if (RemoteFileExists(MetalStdLibPath))
@@ -1402,11 +1394,11 @@ void BuildMetalShaderOutput(
         FString MetalFilePath = (TempDir / HashedName) + TEXT(".metal");
         
 		FString InputFilename = MetalFilePath;
-		FString ObjFilename = FPaths::CreateTempFilename(TempDir, TEXT("ShaderIn"), TEXT(""));
-		FString OutputFilename = FPaths::CreateTempFilename(TempDir, TEXT("ShaderIn"), TEXT(""));
+		FString ObjFilename = FPaths::CreateTempFilename(TempDir, TEXT("ShaderObj"), TEXT(""));
+		FString OutputFilename = FPaths::CreateTempFilename(TempDir, TEXT("ShaderOut"), TEXT(""));
 		
         // write out shader source, the move it into place using an atomic move - ensures only one compile "wins"
-        FString SaveFile = FPaths::CreateTempFilename(TempDir, TEXT("ShaderIn"), TEXT(""));
+        FString SaveFile = FPaths::CreateTempFilename(TempDir, TEXT("ShaderTemp"), TEXT(""));
         FFileHelper::SaveStringToFile(MetalCode, *SaveFile);
         IFileManager::Get().Move(*MetalFilePath, *SaveFile, false, false, true, true);
         IFileManager::Get().Delete(*SaveFile);
@@ -1591,7 +1583,7 @@ void BuildMetalShaderOutput(
 				if (UE4StdLibCRC == 0)
 				{
 					TArrayView<const uint8> UE4PCHData((const uint8*)ue4_stdlib_metal, ue4_stdlib_metal_len);
-					FString UE4StdLibFilename = FPaths::CreateTempFilename(TempDir, TEXT("ShaderIn"), TEXT(""));
+					FString UE4StdLibFilename = FPaths::CreateTempFilename(TempDir, TEXT("ShaderStdLib"), TEXT(""));
 					if (FFileHelper::SaveArrayToFile(UE4PCHData, *UE4StdLibFilename))
 					{
 						FString RemoteTempPath = LocalPathToRemote(UE4StdLibFilename, MakeRemoteTempFolder(FString(TempDir)));
@@ -1606,7 +1598,7 @@ void BuildMetalShaderOutput(
 				if (!RemoteFileExists(RemoteUE4StdLibFilePath) || !ChecksumRemoteFile(*RemoteUE4StdLibFilePath, &RemotePchCRC, &RemotePchLen) || RemotePchCRC != UE4StdLibCRC)
 				{
 					TArrayView<const uint8> UE4PCHData((const uint8*)ue4_stdlib_metal, ue4_stdlib_metal_len);
-					FString UE4StdLibFilename = FPaths::CreateTempFilename(TempDir, TEXT("ShaderIn"), TEXT(""));
+					FString UE4StdLibFilename = FPaths::CreateTempFilename(TempDir, TEXT("ShaderStdLib"), TEXT(""));
 					if (FFileHelper::SaveArrayToFile(UE4PCHData, *UE4StdLibFilename))
 					{
 						IFileManager::Get().Move(*UE4StdLibFilePath, *UE4StdLibFilename, false, false, true, true);
