Fix a silly bug that was putting some RWBuffer objects into the Buffer object list inside the MetalShaderFormat code and then assigning them an invalid index.

diff --git a/Engine/Shaders/Public/MetalCommon.ush b/Engine/Shaders/Public/MetalCommon.ush
--- a/Engine/Shaders/Public/MetalCommon.ush
+++ b/Engine/Shaders/Public/MetalCommon.ush
@@ -5,7 +5,7 @@
 =============================================================================*/
 
 // Update this GUID to improve shader recompilation for Metal only shaders
-// GUID = {812E84A9-7D28-4609-8037-FAA0B58275E0
+// GUID = 25E1A2E4-7259-4763-ABE1-A3C58562D8EF
 
   
 #pragma once
diff --git a/Engine/Source/Developer/Apple/MetalShaderFormat/Private/MetalUtils.h b/Engine/Source/Developer/Apple/MetalShaderFormat/Private/MetalUtils.h
--- a/Engine/Source/Developer/Apple/MetalShaderFormat/Private/MetalUtils.h
+++ b/Engine/Source/Developer/Apple/MetalShaderFormat/Private/MetalUtils.h
@@ -169,13 +169,13 @@ struct FBuffers
 			check(Var);
 			bool bIsStructuredBuffer = (!strncmp(Var->type->name, "RWStructuredBuffer<", 19) || !strncmp(Var->type->name, "StructuredBuffer<", 17));
 			bool bIsByteAddressBuffer = (!strncmp(Var->type->name, "RWByteAddressBuffer", 19) || !strncmp(Var->type->name, "ByteAddressBuffer", 17));
-			if(Var->type->sampler_buffer && Var->type->inner_type->components() != 3 && !Var->invariant && !bIsStructuredBuffer && !bIsByteAddressBuffer)
+			if (Var->type->is_image())
+			{
+				IBuffers.push_back(Var);
+			}
+			else if(Var->type->sampler_buffer && Var->type->inner_type->components() != 3 && !Var->invariant && !bIsStructuredBuffer && !bIsByteAddressBuffer)
             {
                 TBuffers.push_back(Var);
-            }
-            else if (Var->type->is_image())
-            {
-                IBuffers.push_back(Var);
             }
 			else if (Var->semantic && strlen(Var->semantic) == 1)
 			{
