Fixing case where we were mixing functions from different metal compilers. Now we'll blow the cache if the major version doesn't match.

diff --git a/Engine/Source/Developer/Apple/MetalShaderFormat/Private/MetalShaderFormat.cpp b/Engine/Source/Developer/Apple/MetalShaderFormat/Private/MetalShaderFormat.cpp
--- a/Engine/Source/Developer/Apple/MetalShaderFormat/Private/MetalShaderFormat.cpp
+++ b/Engine/Source/Developer/Apple/MetalShaderFormat/Private/MetalShaderFormat.cpp
@@ -186,15 +186,32 @@ uint32 GetMetalFormatVersion(FName Format)
 	{
 		GConfig->GetBool(TEXT("/Script/MacTargetPlatform.MacTargetSettings"), TEXT("XcodeVersionInShaderVersion"), bAddXcodeVersionInShaderVersion, GEngineIni);
 	}
+
+	// We want to include the Xcode App and build version to avoid
+	// weird mismatches where some shaders are built with one version
+	// of the metal frontend and others with a different version.
+	uint64 BuildVersion = 0;
+	
+	// GetXcodeVersion returns:
+	// Major  << 8 | Minor << 4 | Patch
+	AppVersion = GetXcodeVersion(BuildVersion);
+	
 	if (!FApp::IsEngineInstalled() && bAddXcodeVersionInShaderVersion)
 	{
-		uint64 BuildVersion = 0;
-		AppVersion = GetXcodeVersion(BuildVersion);
+		// For local development we'll mix in the xcode version
+		// and build version.
 		AppVersion ^= (BuildVersion & 0xff);
 		AppVersion ^= ((BuildVersion >> 16) & 0xff);
 		AppVersion ^= ((BuildVersion >> 32) & 0xff);
 		AppVersion ^= ((BuildVersion >> 48) & 0xff);
 	}
+	else
+	{
+		// In the other case (ie, shipping editor binary distributions)
+		// We will only mix in the Major version of Xcode used to create
+		// the shader binaries.
+		AppVersion = (AppVersion >> 8) & 0xff;
+	}
 
 	Version.Version.XcodeVersion = AppVersion;
 	Version.Version.Format = FMetalShaderFormat::HEADER_VERSION;
