Fix compiling shaders to Metal that use SceneDepth nodes but do not write any color targets - only depth. If a shader writes nothing at all then it will assume float4 input for the fragment fetch binding.

diff --git a/Engine/Source/Developer/Apple/MetalShaderFormat/Private/MetalUtils.cpp b/Engine/Source/Developer/Apple/MetalShaderFormat/Private/MetalUtils.cpp
--- a/Engine/Source/Developer/Apple/MetalShaderFormat/Private/MetalUtils.cpp
+++ b/Engine/Source/Developer/Apple/MetalShaderFormat/Private/MetalUtils.cpp
@@ -157,6 +157,7 @@ static void CreateNewAssignmentsHalf2Float(_mesa_glsl_parse_state* State, exec_l
 const glsl_type* GetFragColorTypeFromMetalOutputStruct(const glsl_type* OutputType)
 {
 	const glsl_type* FragColorType = glsl_type::error_type;
+	const glsl_type* FragDepthType = glsl_type::error_type;
 	if (OutputType && OutputType->base_type == GLSL_TYPE_STRUCT)
 	{
 		for (unsigned j = 0; j < OutputType->length; j++)
@@ -169,8 +170,16 @@ const glsl_type* GetFragColorTypeFromMetalOutputStruct(const glsl_type* OutputTy
 					FragColorType = OutputType->fields.structure[j].type;
 					break;
 				}
+				else if (!strncmp(OutputType->fields.structure[j].semantic, "[[ depth(", 9))
+				{
+					FragDepthType = OutputType->fields.structure[j].type;
+				}
 			}
 		}
+		if (FragColorType == glsl_type::error_type)
+		{
+			FragColorType = FragDepthType;
+		}
 	}
 	return FragColorType;
 }
@@ -1029,6 +1038,11 @@ struct FFixIntrinsicsVisitor : public ir_rvalue_visitor
 				if (!DestColorVar)
 				{
 					// Generate new input variable for Metal semantics
+					if (DestColorType == glsl_type::error_type)
+					{
+						// When there are no depth writes and no color target writes then use float
+						DestColorType = glsl_type::float_type;
+					}
 					DestColorVar = new(State)ir_variable(glsl_type::get_instance(DestColorType->base_type, 4, 1), "gl_LastFragData", ir_var_in);
 					DestColorVar->semantic = "[[ color(0) ]]";
 				}
