Previous commit overlooked the subtlties that FixIntrinsics was dealing with in addition to FramebufferFetchES2 - so pass through the shader type and only handle FramebufferFetchES2 for fragment shaders instead.

diff --git a/Engine/Source/Developer/Apple/MetalShaderFormat/Private/MetalBackend.cpp b/Engine/Source/Developer/Apple/MetalShaderFormat/Private/MetalBackend.cpp
--- a/Engine/Source/Developer/Apple/MetalShaderFormat/Private/MetalBackend.cpp
+++ b/Engine/Source/Developer/Apple/MetalShaderFormat/Private/MetalBackend.cpp
@@ -4935,12 +4935,8 @@ char* FMetalCodeBackend::GenerateCode(exec_list* ir, _mesa_glsl_parse_state* sta
 	ExpandArrayAssignments(ir, state);
 
 	// Fix any special language extensions (FrameBufferFetchES2() intrinsic)
-	// Only applicable to fragment shaders - this intrinsic is not available elsewhere
-	if (Frequency == HSF_PixelShader)
-	{
-		FixIntrinsics(ir, state);
-	}
-
+	FixIntrinsics(ir, state, Frequency);
+	
 	// Remove half->float->half or float->half->float
 	FixRedundantCasts(ir);
 
diff --git a/Engine/Source/Developer/Apple/MetalShaderFormat/Private/MetalBackend.h b/Engine/Source/Developer/Apple/MetalShaderFormat/Private/MetalBackend.h
--- a/Engine/Source/Developer/Apple/MetalShaderFormat/Private/MetalBackend.h
+++ b/Engine/Source/Developer/Apple/MetalShaderFormat/Private/MetalBackend.h
@@ -102,7 +102,7 @@ struct FMetalCodeBackend : public FCodeBackend
 	void InsertArgumentBuffers(exec_list* ir, _mesa_glsl_parse_state* state, FBuffers& Buffers);
 	void PackInputsAndOutputs(exec_list* ir, _mesa_glsl_parse_state* state, EHlslShaderFrequency Frequency, exec_list& InputVars);
 	void MovePackedUniformsToMain(exec_list* ir, _mesa_glsl_parse_state* state, FBuffers& OutBuffers);
-	void FixIntrinsics(exec_list* ir, _mesa_glsl_parse_state* state);
+	void FixIntrinsics(exec_list* ir, _mesa_glsl_parse_state* state,EHlslShaderFrequency InFrequency);
 	void RemovePackedVarReferences(exec_list* ir, _mesa_glsl_parse_state* State);
 	void PromoteInputsAndOutputsGlobalHalfToFloat(exec_list* ir, _mesa_glsl_parse_state* state, EHlslShaderFrequency Frequency);
 	void ConvertHalfToFloatUniformsAndSamples(exec_list* ir, _mesa_glsl_parse_state* State, bool bConvertUniforms, bool bConvertSamples);
diff --git a/Engine/Source/Developer/Apple/MetalShaderFormat/Private/MetalUtils.cpp b/Engine/Source/Developer/Apple/MetalShaderFormat/Private/MetalUtils.cpp
--- a/Engine/Source/Developer/Apple/MetalShaderFormat/Private/MetalUtils.cpp
+++ b/Engine/Source/Developer/Apple/MetalShaderFormat/Private/MetalUtils.cpp
@@ -952,13 +952,15 @@ struct FFixIntrinsicsVisitor : public ir_rvalue_visitor
 	ir_variable* DestColorVar;
 	const glsl_type* DestColorType;
 	ir_variable* DestMRTColorVar[MAX_SIMULTANEOUS_RENDER_TARGETS];
+	EHlslShaderFrequency Frequency;
 
-	FFixIntrinsicsVisitor(_mesa_glsl_parse_state* InState, ir_function_signature* InMainSig) :
+	FFixIntrinsicsVisitor(_mesa_glsl_parse_state* InState, ir_function_signature* InMainSig, EHlslShaderFrequency InFrequency) :
 		State(InState),
 		bUsesFramebufferFetchES2(false),
 		MRTFetchMask(0),
 		DestColorVar(nullptr),
-		DestColorType(glsl_type::error_type)
+		DestColorType(glsl_type::error_type),
+		Frequency(InFrequency)
 	{
 		DestColorType = GetFragColorTypeFromMetalOutputStruct(InMainSig->return_type);
 		memset(DestMRTColorVar, 0, sizeof(DestMRTColorVar));
@@ -1025,7 +1027,7 @@ struct FFixIntrinsicsVisitor : public ir_rvalue_visitor
 
 	virtual ir_visitor_status visit_leave(ir_call* IR) override
 	{
-		if (IR->use_builtin)
+		if ((Frequency == HSF_PixelShader) && IR->use_builtin)
 		{
 			const char* CalleeName = IR->callee_name();
 			static auto ES2Len = strlen(FRAMEBUFFER_FETCH_ES2);
@@ -1080,12 +1082,12 @@ struct FFixIntrinsicsVisitor : public ir_rvalue_visitor
 	}
 };
 
-void FMetalCodeBackend::FixIntrinsics(exec_list* ir, _mesa_glsl_parse_state* state)
+void FMetalCodeBackend::FixIntrinsics(exec_list* ir, _mesa_glsl_parse_state* state, EHlslShaderFrequency InFrequency)
 {
 	ir_function_signature* MainSig = GetMainFunction(ir);
 	check(MainSig);
 
-	FFixIntrinsicsVisitor Visitor(state,MainSig);
+	FFixIntrinsicsVisitor Visitor(state,MainSig,InFrequency);
 	Visitor.run(&MainSig->body);
 
 	if (Visitor.bUsesFramebufferFetchES2)
