Apply the workaround from CL #3722943 to the patchCount variable as well in the hope that this will fix Nvidia PSO compiler failures - it fixes at least one of the examples we have logged.

diff --git a/Engine/Shaders/Public/Platform/Metal/MetalCommon.ush b/Engine/Shaders/Public/Platform/Metal/MetalCommon.ush
--- a/Engine/Shaders/Public/Platform/Metal/MetalCommon.ush
+++ b/Engine/Shaders/Public/Platform/Metal/MetalCommon.ush
@@ -5,7 +5,7 @@
 =============================================================================*/
 
 // Update this GUID to improve shader recompilation for Metal only shaders
-// GUID = 4E19BD71-3C18-4138-B085-451EAC550287
+// GUID = 458545C3-5894-4AAB-9E80-65235AA26706
   
 #pragma once
 
diff --git a/Engine/Source/Developer/Apple/MetalShaderFormat/Private/MetalBackend.cpp b/Engine/Source/Developer/Apple/MetalShaderFormat/Private/MetalBackend.cpp
--- a/Engine/Source/Developer/Apple/MetalShaderFormat/Private/MetalBackend.cpp
+++ b/Engine/Source/Developer/Apple/MetalShaderFormat/Private/MetalBackend.cpp
@@ -1483,8 +1483,8 @@ protected:
 		{
 			check(sig->is_main);
             
-            ir_variable* patchCount = new(ParseState)ir_variable(glsl_type::uint_type, "patchCount", ir_var_uniform);
-            patchCount->semantic = "u";
+            ir_variable* patchCount = new(ParseState)ir_variable(glsl_type::uint_type, "patchCount", ir_var_in);
+            patchCount->semantic = "";
             Buffers.Buffers.Add(patchCount);
             
             int32 patchIndex = Buffers.GetIndex(patchCount);
@@ -1502,7 +1502,7 @@ protected:
 				"uint2 thread_position_in_grid [[thread_position_in_grid]],\n"
 				"ushort2 thread_position_in_threadgroup [[thread_position_in_threadgroup]],\n"
 				"uint2 threadgroup_position_in_grid [[threadgroup_position_in_grid]],\n"
-				"constant uint *patchCount [[ buffer(%d) ]],\n"
+				"device const uint *patchCount [[ buffer(%d) ]],\n"
 				"#define METAL_INDEX_BUFFER_ID %d\n"
 				"const device typed_buffer<uint>* indexBuffer [[ buffer(METAL_INDEX_BUFFER_ID) ]]",
                 patchIndex, IndexBufferIndex
