Duplicate 5263216 to aid debugging Metal shader issues in games using native shader libraries without having to cook locally:

Package Metal shader source into a zip file rather than a tgz so it can be done on Windows builds too and do this asynchronously while generating the Metal libraries. This file is stored in the MetaData folder so should be moved out of the content and not get packaged. Must be unzipped at the command-line for some reason, but it works.

diff --git a/Engine/Source/Developer/Apple/MetalShaderFormat/MetalShaderFormat.Build.cs b/Engine/Source/Developer/Apple/MetalShaderFormat/MetalShaderFormat.Build.cs
--- a/Engine/Source/Developer/Apple/MetalShaderFormat/MetalShaderFormat.Build.cs
+++ b/Engine/Source/Developer/Apple/MetalShaderFormat/MetalShaderFormat.Build.cs
@@ -20,7 +20,8 @@ public class MetalShaderFormat : ModuleRules
 				"Core",
 				"RenderCore",
 				"ShaderCompilerCommon",
-				"ShaderPreprocessor"
+				"ShaderPreprocessor",
+				"FileUtilities"
 			}
 			);
 
diff --git a/Engine/Source/Developer/Apple/MetalShaderFormat/Private/MetalShaderFormat.cpp b/Engine/Source/Developer/Apple/MetalShaderFormat/Private/MetalShaderFormat.cpp
--- a/Engine/Source/Developer/Apple/MetalShaderFormat/Private/MetalShaderFormat.cpp
+++ b/Engine/Source/Developer/Apple/MetalShaderFormat/Private/MetalShaderFormat.cpp
@@ -10,9 +10,12 @@
 #include "hlslcc.h"
 #include "MetalShaderResources.h"
 #include "HAL/FileManager.h"
+#include "HAL/PlatformFilemanager.h"
 #include "Serialization/Archive.h"
 #include "Misc/ConfigCacheIni.h"
 #include "MetalBackend.h"
+#include "Misc/FileHelper.h"
+#include "FileUtilities/ZipArchiveWriter.h"
 
 #define WRITE_METAL_SHADER_SOURCE_ARCHIVE 0
 
@@ -113,6 +116,54 @@ public:
 
 			Tasks.Add(CompletionFence);
 		}
+		
+#if WITH_ENGINE
+		FGraphEventRef DebugDataCompletionFence = FFunctionGraphTask::CreateAndDispatchWhenReady([this, OutputDir, LibraryPlatformName, DebugOutputDir]()
+		{
+			//TODO add a check in here - this will only work if we have shader archiving with debug info set.
+			
+			//We want to archive all the metal shader source files so that they can be unarchived into a debug location
+			//This allows the debugging of optimised metal shaders within the xcode tool set
+			//Currently using the 'tar' system tool to create a compressed tape archive
+			
+			//Place the archive in the same position as the .metallib file
+			FString CompressedDir = (OutputDir / TEXT("../MetaData/ShaderDebug/"));
+			IFileManager::Get().MakeDirectory(*CompressedDir, true);
+			
+			FString CompressedPath = (CompressedDir / LibraryPlatformName) + TEXT(".zip");
+			
+			IPlatformFile& PlatformFile = FPlatformFileManager::Get().GetPlatformFile();
+			IFileHandle* ZipFile = PlatformFile.OpenWrite(*CompressedPath);
+			if (ZipFile)
+			{
+				FZipArchiveWriter* ZipWriter = new FZipArchiveWriter(ZipFile);
+				
+				//Find the metal source files
+				TArray<FString> FilesToArchive;
+				IFileManager::Get().FindFilesRecursive( FilesToArchive, *DebugOutputDir, TEXT("*.metal"), true, false, false );
+				
+				//Write the local file names into the target file
+				const FString DebugDir = DebugOutputDir / *Format.GetPlainNameString();
+				
+				for(FString FileName : FilesToArchive)
+				{
+					TArray<uint8> FileData;
+					FFileHelper::LoadFileToArray(FileData, *FileName);
+					FPaths::MakePathRelativeTo(FileName, *DebugDir);
+					
+					ZipWriter->AddFile(FileName, FileData, FDateTime::Now());
+				}
+				
+				delete ZipWriter;
+				ZipWriter = nullptr;
+			}
+			else
+			{
+				UE_LOG(LogShaders, Error, TEXT("Failed to create Metal debug .zip output file \"%s\". Debug .zip export will be disabled."), *CompressedPath);
+			}
+		}, TStatId(), NULL, ENamedThreads::AnyThread);
+		Tasks.Add(DebugDataCompletionFence);
+#endif
 
 		// Wait for tasks
 		for (auto& Task : Tasks)
@@ -138,85 +189,6 @@ public:
 				
 				bOK = true;
 			}
-
-#if PLATFORM_MAC && WRITE_METAL_SHADER_SOURCE_ARCHIVE
-			if(bOK)
-			{
-				//TODO add a check in here - this will only work if we have shader archiving with debug info set.
-				
-				//We want to archive all the metal shader source files so that they can be unarchived into a debug location
-				//This allows the debugging of optimised metal shaders within the xcode tool set
-				//Currently using the 'tar' system tool to create a compressed tape archive
-				
-				//Place the archive in the same position as the .metallib file
-				FString CompressedPath = (OutputDir / LibraryPlatformName) + TEXT(".tgz");
-				
-				FString ArchiveCommand = TEXT("/usr/bin/tar");
-				
-				// Iterative support for pre-stripped shaders - unpack existing tgz archive without file overwrite - if it exists in cooked dir we're in iterative mode
-				if(FPaths::FileExists(CompressedPath))
-				{
-					int32 ReturnCode = -1;
-					FString Result;
-					FString Errors;
-					
-					FString ExtractCommandParams = FString::Printf(TEXT("xopfk \"%s\" -C \"%s\""), *CompressedPath, *DebugOutputDir);
-					FPlatformProcess::ExecProcess( *ArchiveCommand, *ExtractCommandParams, &ReturnCode, &Result, &Errors );
-				}
-				
-				//Due to the limitations of the 'tar' command and running through NSTask,
-				//the most reliable way is to feed it a list of local file name (-T) with a working path set (-C)
-				//if we built the list with absolute paths without -C then we'd get the full folder structure in the archive
-				//I don't think we want this here
-				
-				//Build a file list that 'tar' can access
-				const FString FileListPath = DebugOutputDir / TEXT("ArchiveInput.txt");
-				IFileManager::Get().Delete( *FileListPath );
-				
-				{
-					//Find the metal source files
-					TArray<FString> FilesToArchive;
-					IFileManager::Get().FindFilesRecursive( FilesToArchive, *DebugOutputDir, TEXT("*.metal"), true, false, false );
-					
-					//Write the local file names into the target file
-					FArchive* FileListHandle = IFileManager::Get().CreateFileWriter( *FileListPath );
-					if(FileListHandle)
-					{
-						const FString NewLine = TEXT("\n");
-						
-						const FString DebugDir = DebugOutputDir / *Format.GetPlainNameString();
-						
-						for(FString FileName : FilesToArchive)
-						{
-							FPaths::MakePathRelativeTo(FileName, *DebugDir);
-							
-							FString TextLine = FileName + NewLine;
-							
-							//We don't want the string to archive through the << operator otherwise we'd be creating a binary file - we need text
-							auto AnsiFullPath = StringCast<ANSICHAR>( *TextLine );
-							FileListHandle->Serialize( (ANSICHAR*)AnsiFullPath.Get(), AnsiFullPath.Length() );
-						}
-						
-						//Clean up
-						FileListHandle->Close();
-						delete FileListHandle;
-					}
-				}
-				
-				int32 ReturnCode = -1;
-				FString Result;
-				FString Errors;
-				
-				//Setup the NSTask command and parameter list, Archive (-c) and Compress (-z) to target file (-f) the metal file list (-T) using a local dir in archive (-C).
-				FString ArchiveCommandParams = FString::Printf( TEXT("czf \"%s\" -C \"%s\" -T \"%s\""), *CompressedPath, *DebugOutputDir, *FileListPath );
-				
-				//Execute command, this should end up with a .tgz file in the same location at the .metallib file
-				if(!FPlatformProcess::ExecProcess( *ArchiveCommand, *ArchiveCommandParams, &ReturnCode, &Result, &Errors ) || ReturnCode != 0)
-				{
-					UE_LOG(LogShaders, Error, TEXT("Archive Shader Source failed %d: %s"), ReturnCode, *Errors);
-				}
-			}
-#endif
 		}
 
 		return bOK;
