Fix Metal shader translation bugs:
- Texture*Array.Load needs to split the array index out of the vector coordinate for Metal.
- The GetDimensions function needs to support the array-length/depth output argument.
- RemovePackedVarReferences was inserting new variables at the top of the function, not the same scope as the source variable, which won't compile in cases where structs are used as temporaries and not function arguments.

diff --git a/Engine/Source/Developer/Apple/MetalShaderFormat/Private/MetalBackend.cpp b/Engine/Source/Developer/Apple/MetalShaderFormat/Private/MetalBackend.cpp
--- a/Engine/Source/Developer/Apple/MetalShaderFormat/Private/MetalBackend.cpp
+++ b/Engine/Source/Developer/Apple/MetalShaderFormat/Private/MetalBackend.cpp
@@ -2034,7 +2034,8 @@ protected:
 				}
 				break;
 			case ir_txs:
-				ralloc_asprintf_append(buffer, "int2((int)");
+				print_type_pre(tex->type);
+				ralloc_asprintf_append(buffer, "((int)");
 				break;
 			default:
 				break;
@@ -2310,7 +2311,57 @@ protected:
 			{
 				tex->sampler->accept(this);
 				ralloc_asprintf_append(buffer, ".read(");
-				tex->coordinate->accept(this);
+				
+				if (tex->sampler->type->sampler_array)
+				{
+					// Need to split the coordinate
+					char const* CoordSwizzle = "";
+					char const* IndexSwizzle = "y";
+					switch(tex->sampler->type->sampler_dimensionality)
+					{
+						case GLSL_SAMPLER_DIM_1D:
+						{
+							break;
+						}
+						case GLSL_SAMPLER_DIM_2D:
+						case GLSL_SAMPLER_DIM_RECT:
+						{
+							CoordSwizzle = "y";
+							IndexSwizzle = "z";
+							break;
+						}
+						case GLSL_SAMPLER_DIM_3D:
+						{
+							CoordSwizzle = "yz";
+							IndexSwizzle = "w";
+							break;
+						}
+						case GLSL_SAMPLER_DIM_CUBE:
+						{
+							CoordSwizzle = "yz";
+							IndexSwizzle = "w";
+							break;
+						}
+						case GLSL_SAMPLER_DIM_BUF:
+						case GLSL_SAMPLER_DIM_EXTERNAL:
+						default:
+						{
+							check(0);
+							break;
+						}
+					}
+					
+					ralloc_asprintf_append(buffer, "(");
+					tex->coordinate->accept(this);
+					
+					ralloc_asprintf_append(buffer, ").x%s, (uint)(", CoordSwizzle);
+					tex->coordinate->accept(this);
+					ralloc_asprintf_append(buffer, ").%s", IndexSwizzle);
+				}
+				else
+				{
+					tex->coordinate->accept(this);
+				}
 				
 				if (tex->sampler->type->sampler_ms)
 				{
@@ -2482,6 +2533,7 @@ protected:
 			{
 				ralloc_asprintf_append(buffer, tex->sampler->type->sampler_shadow ? "depth_cube_array::" : "texture_cube_array::");
 			}
+			else
 			{
 				tex->sampler->accept(this);
 				ralloc_asprintf_append(buffer, ".");
@@ -2498,6 +2550,65 @@ protected:
 				tex->lod_info.lod->accept(this);
 			}
 			ralloc_asprintf_append(buffer, ")");
+			
+			if (tex->type->vector_elements == 3)
+			{
+				switch(tex->sampler->type->sampler_dimensionality)
+				{
+					case GLSL_SAMPLER_DIM_1D:
+					{
+						break;
+					}
+					case GLSL_SAMPLER_DIM_2D:
+					case GLSL_SAMPLER_DIM_RECT:
+					{
+						if (tex->sampler->type->sampler_array)
+						{
+							ralloc_asprintf_append(buffer, ", (int)");
+							tex->sampler->accept(this);
+							ralloc_asprintf_append(buffer, ".get_array_size()");
+						}
+						else
+						{
+							check(0);
+						}
+						break;
+					}
+					case GLSL_SAMPLER_DIM_3D:
+					{
+						ralloc_asprintf_append(buffer, ", (int)");
+						tex->sampler->accept(this);
+						ralloc_asprintf_append(buffer, ".get_depth(");
+						if (tex->lod_info.lod)
+						{
+							tex->lod_info.lod->accept(this);
+						}
+						ralloc_asprintf_append(buffer, ")");
+						break;
+					}
+					case GLSL_SAMPLER_DIM_CUBE:
+					{
+						if (tex->sampler->type->sampler_array)
+						{
+							ralloc_asprintf_append(buffer, tex->sampler->type->sampler_shadow ? ", depth_cube_array::get_array_size(" : ", texture_cube_array::get_array_size(");
+							tex->sampler->accept(this);
+							ralloc_asprintf_append(buffer, ")");
+						}
+						else
+						{
+							check(0);
+						}
+						break;
+					}
+					case GLSL_SAMPLER_DIM_BUF:
+					case GLSL_SAMPLER_DIM_EXTERNAL:
+					default:
+					{
+						check(0);
+						break;
+					}
+				}
+			}
 		}
 			break;
 
@@ -2685,7 +2796,59 @@ protected:
 					ralloc_asprintf_append(buffer, "(");
 					deref->image->accept(this);
 					ralloc_asprintf_append(buffer, ".read(");
-					deref->image_index->accept(this);
+					
+					if (bIsArray)
+					{
+						// Need to split the coordinate
+						char const* CoordSwizzle = "";
+						char const* IndexSwizzle = "y";
+						switch(deref->image->type->sampler_dimensionality)
+						{
+							case GLSL_SAMPLER_DIM_1D:
+							{
+								break;
+							}
+							case GLSL_SAMPLER_DIM_2D:
+							case GLSL_SAMPLER_DIM_RECT:
+							{
+								CoordSwizzle = "y";
+								IndexSwizzle = "z";
+								break;
+							}
+							case GLSL_SAMPLER_DIM_3D:
+							{
+								CoordSwizzle = "yz";
+								IndexSwizzle = "w";
+								break;
+							}
+							case GLSL_SAMPLER_DIM_CUBE:
+							{
+								CoordSwizzle = "yz";
+								IndexSwizzle = "w";
+								break;
+							}
+							case GLSL_SAMPLER_DIM_BUF:
+							case GLSL_SAMPLER_DIM_EXTERNAL:
+							default:
+							{
+								check(0);
+								break;
+							}
+						}
+						
+						ralloc_asprintf_append(buffer, "(");
+						deref->image_index->accept(this);
+						
+						ralloc_asprintf_append(buffer, ").x%s, (uint)(", CoordSwizzle);
+						deref->image_index->accept(this);
+						ralloc_asprintf_append(buffer, ").%s", IndexSwizzle);
+					}
+					else
+					{
+						deref->image_index->accept(this);
+					}
+					
+					
 					ralloc_asprintf_append(buffer, ")");
 					switch(dst_elements)
 					{
@@ -2987,7 +3150,8 @@ protected:
 			//		Temp = textureSize(T{, lod});
 			// Metal
 			//		int2 Temp = int2((int)T.get_width({lod}), (int)T.get_height({lod}));
-			ralloc_asprintf_append(buffer, "int2(");
+			print_type_pre(deref->type);
+			ralloc_asprintf_append(buffer, "((int)");
 			deref->image->accept(this);
 			ralloc_asprintf_append(buffer, ".get_width(");
 			
@@ -3003,7 +3167,67 @@ protected:
 			{
 				deref->image_index->accept(this);
 			}
-			ralloc_asprintf_append(buffer, "))");
+			ralloc_asprintf_append(buffer, ")");
+			
+			if (deref->type->vector_elements == 3)
+			{
+				switch(deref->image->type->sampler_dimensionality)
+				{
+					case GLSL_SAMPLER_DIM_1D:
+					{
+						break;
+					}
+					case GLSL_SAMPLER_DIM_2D:
+					case GLSL_SAMPLER_DIM_RECT:
+					{
+						if (deref->image->type->sampler_array)
+						{
+							ralloc_asprintf_append(buffer, ", (int)");
+							deref->image->accept(this);
+							ralloc_asprintf_append(buffer, ".get_array_size()");
+						}
+						else
+						{
+							check(0);
+						}
+						break;
+					}
+					case GLSL_SAMPLER_DIM_3D:
+					{
+						ralloc_asprintf_append(buffer, ", (int)");
+						deref->image->accept(this);
+						ralloc_asprintf_append(buffer, ".get_depth(");
+						if (deref->image_index)
+						{
+							deref->image_index->accept(this);
+						}
+						ralloc_asprintf_append(buffer, ")");
+						break;
+					}
+					case GLSL_SAMPLER_DIM_CUBE:
+					{
+						if (deref->image->type->sampler_array)
+						{
+							ralloc_asprintf_append(buffer, deref->image->type->sampler_shadow ? ", depth_cube_array::get_array_size(" : ", texture_cube_array::get_array_size(");
+							deref->image->accept(this);
+							ralloc_asprintf_append(buffer, ")");
+						}
+						else
+						{
+							check(0);
+						}
+						break;
+					}
+					case GLSL_SAMPLER_DIM_BUF:
+					case GLSL_SAMPLER_DIM_EXTERNAL:
+					default:
+					{
+						check(0);
+						break;
+					}
+				}
+			}
+			ralloc_asprintf_append(buffer, ")");
 		}
 		else
 		{
diff --git a/Engine/Source/Developer/Apple/MetalShaderFormat/Private/MetalUtils.cpp b/Engine/Source/Developer/Apple/MetalShaderFormat/Private/MetalUtils.cpp
--- a/Engine/Source/Developer/Apple/MetalShaderFormat/Private/MetalUtils.cpp
+++ b/Engine/Source/Developer/Apple/MetalShaderFormat/Private/MetalUtils.cpp
@@ -3387,7 +3387,7 @@ struct FDeReferencePackedVarsVisitor final : public ir_rvalue_visitor
 			{
 				if (SwizzleValDeRefRecord->type->vector_elements > 1 && SwizzleValDeRefRecord->type->vector_elements < 4)
 				{
-					auto* Var = GetVar(SwizzleValDeRefRecord);
+					auto* Var = GetVar(StructVar, SwizzleValDeRefRecord);
 					Swizzle->val = new(State)ir_dereference_variable(Var);
 				}
 			}
@@ -3399,7 +3399,7 @@ struct FDeReferencePackedVarsVisitor final : public ir_rvalue_visitor
 			{
 				if (DeRefRecord->type->vector_elements > 1 && DeRefRecord->type->vector_elements < 4)
 				{
-					auto* Var = GetVar(DeRefRecord);
+					auto* Var = GetVar(StructVar, DeRefRecord);
 					*RValuePtr = new(State)ir_dereference_variable(Var);
 				}
             }
@@ -3407,7 +3407,7 @@ struct FDeReferencePackedVarsVisitor final : public ir_rvalue_visitor
             {
                 if (DeRefRecord->type->vector_elements > 1 && DeRefRecord->type->vector_elements < 4)
                 {
-                    auto* Var = GetVar(DeRefRecord);
+                    auto* Var = GetVar(StructVar, DeRefRecord);
                     *RValuePtr = new(State)ir_dereference_variable(Var);
                 }
             }
@@ -3415,8 +3415,9 @@ struct FDeReferencePackedVarsVisitor final : public ir_rvalue_visitor
 	}
 
 	std::map<ir_dereference_record*, ir_variable*, ir_type_compare<ir_dereference_record>> Replaced;
+	std::map<ir_variable*, ir_variable*, ir_type_compare<ir_variable>> Replacements;
 
-	ir_variable* GetVar(ir_dereference_record* ir)
+	ir_variable* GetVar(ir_variable* Orig, ir_dereference_record* ir)
 	{
 		ir_variable* Var = nullptr;
 		for (auto& Pair : Replaced)
@@ -3432,6 +3433,7 @@ struct FDeReferencePackedVarsVisitor final : public ir_rvalue_visitor
 		{
 			Var = new(State)ir_variable(ir->type, nullptr, ir_var_temporary);
 			Replaced[ir] = Var;
+			Replacements[Var] = Orig;
 		}
 		return Var;
 	}
@@ -3446,15 +3448,28 @@ void FMetalCodeBackend::RemovePackedVarReferences(exec_list* ir, _mesa_glsl_pars
 	{
 		return;
 	}
-
+	
 	ir_function_signature* Main = GetMainFunction(ir);
+	for (auto& Pair : Visitor.Replacements)
+	{
+		auto* NewVar = Pair.first;
+		auto* OldVar = Pair.second;
+		if (OldVar->mode == ir_var_uniform || OldVar->mode == ir_var_in || OldVar->mode == ir_var_out)
+		{
+			Main->body.push_head(NewVar);
+		}
+		else
+		{
+			OldVar->insert_after(NewVar);
+		}
+	}
+
 	for (auto& OuterPair : Visitor.Replaced)
 	{
 		auto* NewVar = OuterPair.second;
 		auto* DeRefRecord = OuterPair.first;
 		auto* NewAssignment = new(State)ir_assignment(new(State)ir_dereference_variable(NewVar), DeRefRecord);
-		Main->body.push_head(NewAssignment);
-		Main->body.push_head(NewVar);
+		NewVar->insert_after(NewAssignment);
 	}
 }
 
