MetalBackend changes to actually convert local constant arrays to the new safe_array template which sanitises input addresses - because iOS doesn't do that for you. This should reduce the number of bad-access GPU timeouts according to Apple.

diff --git a/Engine/Source/Developer/Apple/MetalShaderFormat/Private/MetalBackend.cpp b/Engine/Source/Developer/Apple/MetalShaderFormat/Private/MetalBackend.cpp
--- a/Engine/Source/Developer/Apple/MetalShaderFormat/Private/MetalBackend.cpp
+++ b/Engine/Source/Developer/Apple/MetalShaderFormat/Private/MetalBackend.cpp
@@ -1365,6 +1365,13 @@ protected:
 						ralloc_asprintf_append(buffer, "/*ir_var_out, is_patch_constant*/");
 					}
 				}
+				else if (var->mode == ir_var_auto && var->type->is_array())
+				{
+					ralloc_asprintf_append(buffer, "ue4::safe_array<");
+					print_type_pre(var->type->element_type());
+					ralloc_asprintf_append(buffer, ", %u>", var->type->length);
+					ralloc_asprintf_append(buffer, " %s", unique_name(var));
+				}
 				else
 				{
 					if (var->mode == ir_var_shared)
