Invalidating the MetalShaderFormat version seems to cure the build machine of generating bogus shaders for the BoneMatrix SRV change.

diff --git a/Engine/Source/Developer/Apple/MetalShaderFormat/Private/MetalShaderFormat.cpp b/Engine/Source/Developer/Apple/MetalShaderFormat/Private/MetalShaderFormat.cpp
--- a/Engine/Source/Developer/Apple/MetalShaderFormat/Private/MetalShaderFormat.cpp
+++ b/Engine/Source/Developer/Apple/MetalShaderFormat/Private/MetalShaderFormat.cpp
@@ -213,7 +213,7 @@ class FMetalShaderFormat : public IShaderFormat
 public:
 	enum
 	{
-		HEADER_VERSION = 64,
+		HEADER_VERSION = 65,
 	};
 	
 	struct FVersion
