When compiling Metal shaders sometimes local temporary arrays are actually declared as variables - these should use safe_array too.

diff --git a/Engine/Source/Developer/Apple/MetalShaderFormat/Private/MetalBackend.cpp b/Engine/Source/Developer/Apple/MetalShaderFormat/Private/MetalBackend.cpp
--- a/Engine/Source/Developer/Apple/MetalShaderFormat/Private/MetalBackend.cpp
+++ b/Engine/Source/Developer/Apple/MetalShaderFormat/Private/MetalBackend.cpp
@@ -1365,7 +1365,7 @@ protected:
 						ralloc_asprintf_append(buffer, "/*ir_var_out, is_patch_constant*/");
 					}
 				}
-				else if (var->mode == ir_var_auto && var->type->is_array())
+				else if ((var->mode == ir_var_auto || var->mode == ir_var_temporary) && var->type->is_array())
 				{
 					ralloc_asprintf_append(buffer, "ue4::safe_array<");
 					print_type_pre(var->type->element_type());
@@ -2041,8 +2041,7 @@ protected:
 				}
 				break;
 			case ir_txs:
-				print_type_pre(tex->type);
-				ralloc_asprintf_append(buffer, "((int)");
+				ralloc_asprintf_append(buffer, "int2((int)");
 				break;
 			default:
 				break;
@@ -2318,57 +2317,7 @@ protected:
 			{
 				tex->sampler->accept(this);
 				ralloc_asprintf_append(buffer, ".read(");
-				
-				if (tex->sampler->type->sampler_array)
-				{
-					// Need to split the coordinate
-					char const* CoordSwizzle = "";
-					char const* IndexSwizzle = "y";
-					switch(tex->sampler->type->sampler_dimensionality)
-					{
-						case GLSL_SAMPLER_DIM_1D:
-						{
-							break;
-						}
-						case GLSL_SAMPLER_DIM_2D:
-						case GLSL_SAMPLER_DIM_RECT:
-						{
-							CoordSwizzle = "y";
-							IndexSwizzle = "z";
-							break;
-						}
-						case GLSL_SAMPLER_DIM_3D:
-						{
-							CoordSwizzle = "yz";
-							IndexSwizzle = "w";
-							break;
-						}
-						case GLSL_SAMPLER_DIM_CUBE:
-						{
-							CoordSwizzle = "yz";
-							IndexSwizzle = "w";
-							break;
-						}
-						case GLSL_SAMPLER_DIM_BUF:
-						case GLSL_SAMPLER_DIM_EXTERNAL:
-						default:
-						{
-							check(0);
-							break;
-						}
-					}
-					
-					ralloc_asprintf_append(buffer, "(");
-					tex->coordinate->accept(this);
-					
-					ralloc_asprintf_append(buffer, ").x%s, (uint)(", CoordSwizzle);
-					tex->coordinate->accept(this);
-					ralloc_asprintf_append(buffer, ").%s", IndexSwizzle);
-				}
-				else
-				{
-					tex->coordinate->accept(this);
-				}
+				tex->coordinate->accept(this);
 				
 				if (tex->sampler->type->sampler_ms)
 				{
@@ -2540,7 +2489,6 @@ protected:
 			{
 				ralloc_asprintf_append(buffer, tex->sampler->type->sampler_shadow ? "depth_cube_array::" : "texture_cube_array::");
 			}
-			else
 			{
 				tex->sampler->accept(this);
 				ralloc_asprintf_append(buffer, ".");
@@ -2557,65 +2505,6 @@ protected:
 				tex->lod_info.lod->accept(this);
 			}
 			ralloc_asprintf_append(buffer, ")");
-			
-			if (tex->type->vector_elements == 3)
-			{
-				switch(tex->sampler->type->sampler_dimensionality)
-				{
-					case GLSL_SAMPLER_DIM_1D:
-					{
-						break;
-					}
-					case GLSL_SAMPLER_DIM_2D:
-					case GLSL_SAMPLER_DIM_RECT:
-					{
-						if (tex->sampler->type->sampler_array)
-						{
-							ralloc_asprintf_append(buffer, ", (int)");
-							tex->sampler->accept(this);
-							ralloc_asprintf_append(buffer, ".get_array_size()");
-						}
-						else
-						{
-							check(0);
-						}
-						break;
-					}
-					case GLSL_SAMPLER_DIM_3D:
-					{
-						ralloc_asprintf_append(buffer, ", (int)");
-						tex->sampler->accept(this);
-						ralloc_asprintf_append(buffer, ".get_depth(");
-						if (tex->lod_info.lod)
-						{
-							tex->lod_info.lod->accept(this);
-						}
-						ralloc_asprintf_append(buffer, ")");
-						break;
-					}
-					case GLSL_SAMPLER_DIM_CUBE:
-					{
-						if (tex->sampler->type->sampler_array)
-						{
-							ralloc_asprintf_append(buffer, tex->sampler->type->sampler_shadow ? ", depth_cube_array::get_array_size(" : ", texture_cube_array::get_array_size(");
-							tex->sampler->accept(this);
-							ralloc_asprintf_append(buffer, ")");
-						}
-						else
-						{
-							check(0);
-						}
-						break;
-					}
-					case GLSL_SAMPLER_DIM_BUF:
-					case GLSL_SAMPLER_DIM_EXTERNAL:
-					default:
-					{
-						check(0);
-						break;
-					}
-				}
-			}
 		}
 			break;
 
@@ -2803,59 +2692,7 @@ protected:
 					ralloc_asprintf_append(buffer, "(");
 					deref->image->accept(this);
 					ralloc_asprintf_append(buffer, ".read(");
-					
-					if (bIsArray)
-					{
-						// Need to split the coordinate
-						char const* CoordSwizzle = "";
-						char const* IndexSwizzle = "y";
-						switch(deref->image->type->sampler_dimensionality)
-						{
-							case GLSL_SAMPLER_DIM_1D:
-							{
-								break;
-							}
-							case GLSL_SAMPLER_DIM_2D:
-							case GLSL_SAMPLER_DIM_RECT:
-							{
-								CoordSwizzle = "y";
-								IndexSwizzle = "z";
-								break;
-							}
-							case GLSL_SAMPLER_DIM_3D:
-							{
-								CoordSwizzle = "yz";
-								IndexSwizzle = "w";
-								break;
-							}
-							case GLSL_SAMPLER_DIM_CUBE:
-							{
-								CoordSwizzle = "yz";
-								IndexSwizzle = "w";
-								break;
-							}
-							case GLSL_SAMPLER_DIM_BUF:
-							case GLSL_SAMPLER_DIM_EXTERNAL:
-							default:
-							{
-								check(0);
-								break;
-							}
-						}
-						
-						ralloc_asprintf_append(buffer, "(");
-						deref->image_index->accept(this);
-						
-						ralloc_asprintf_append(buffer, ").x%s, (uint)(", CoordSwizzle);
-						deref->image_index->accept(this);
-						ralloc_asprintf_append(buffer, ").%s", IndexSwizzle);
-					}
-					else
-					{
-						deref->image_index->accept(this);
-					}
-					
-					
+					deref->image_index->accept(this);
 					ralloc_asprintf_append(buffer, ")");
 					switch(dst_elements)
 					{
@@ -3157,8 +2994,7 @@ protected:
 			//		Temp = textureSize(T{, lod});
 			// Metal
 			//		int2 Temp = int2((int)T.get_width({lod}), (int)T.get_height({lod}));
-			print_type_pre(deref->type);
-			ralloc_asprintf_append(buffer, "((int)");
+			ralloc_asprintf_append(buffer, "int2(");
 			deref->image->accept(this);
 			ralloc_asprintf_append(buffer, ".get_width(");
 			
@@ -3174,67 +3010,7 @@ protected:
 			{
 				deref->image_index->accept(this);
 			}
-			ralloc_asprintf_append(buffer, ")");
-			
-			if (deref->type->vector_elements == 3)
-			{
-				switch(deref->image->type->sampler_dimensionality)
-				{
-					case GLSL_SAMPLER_DIM_1D:
-					{
-						break;
-					}
-					case GLSL_SAMPLER_DIM_2D:
-					case GLSL_SAMPLER_DIM_RECT:
-					{
-						if (deref->image->type->sampler_array)
-						{
-							ralloc_asprintf_append(buffer, ", (int)");
-							deref->image->accept(this);
-							ralloc_asprintf_append(buffer, ".get_array_size()");
-						}
-						else
-						{
-							check(0);
-						}
-						break;
-					}
-					case GLSL_SAMPLER_DIM_3D:
-					{
-						ralloc_asprintf_append(buffer, ", (int)");
-						deref->image->accept(this);
-						ralloc_asprintf_append(buffer, ".get_depth(");
-						if (deref->image_index)
-						{
-							deref->image_index->accept(this);
-						}
-						ralloc_asprintf_append(buffer, ")");
-						break;
-					}
-					case GLSL_SAMPLER_DIM_CUBE:
-					{
-						if (deref->image->type->sampler_array)
-						{
-							ralloc_asprintf_append(buffer, deref->image->type->sampler_shadow ? ", depth_cube_array::get_array_size(" : ", texture_cube_array::get_array_size(");
-							deref->image->accept(this);
-							ralloc_asprintf_append(buffer, ")");
-						}
-						else
-						{
-							check(0);
-						}
-						break;
-					}
-					case GLSL_SAMPLER_DIM_BUF:
-					case GLSL_SAMPLER_DIM_EXTERNAL:
-					default:
-					{
-						check(0);
-						break;
-					}
-				}
-			}
-			ralloc_asprintf_append(buffer, ")");
+			ralloc_asprintf_append(buffer, "))");
 		}
 		else
 		{
