Fix for distance field text rendering in Metal - fragmenmts were getting discarded in the pixel shader as the FSimpleElementDistanceFieldGammaPS::FSimpleElementMaskedGammaBasePS::ClipRef value was getting overwritten by a different value in the global structure buffer.  Metal buffer has different alignment for < 4 byte types compared to the expected value in the parameter map.  This change requests a 4 byte alignment for bool inside structs.

diff --git a/Engine/Shaders/Public/Platform/Metal/MetalCommon.ush b/Engine/Shaders/Public/Platform/Metal/MetalCommon.ush
--- a/Engine/Shaders/Public/Platform/Metal/MetalCommon.ush
+++ b/Engine/Shaders/Public/Platform/Metal/MetalCommon.ush
@@ -5,7 +5,7 @@
 =============================================================================*/
 
 // Update this GUID to improve shader recompilation for Metal only shaders
-// GUID = 8648701DC27A4D5E8DC935DBAAB43231
+// GUID = 7625307457499D7614A9DD98DD45F786
   
 #pragma once
 
diff --git a/Engine/Source/Developer/Apple/MetalShaderFormat/Private/MetalBackend.cpp b/Engine/Source/Developer/Apple/MetalShaderFormat/Private/MetalBackend.cpp
--- a/Engine/Source/Developer/Apple/MetalShaderFormat/Private/MetalBackend.cpp
+++ b/Engine/Source/Developer/Apple/MetalShaderFormat/Private/MetalBackend.cpp
@@ -3441,6 +3441,13 @@ protected:
 				for (unsigned j = 0; j < s->length; j++)
 				{
 					ralloc_asprintf_append(buffer, "\t");
+					
+					// HLSL bool is 4 bytes we need to align ours to that as well in structures - do the same for half/short etc?
+					if(GLSL_TYPE_BOOL == s->fields.structure[j].type->base_type)
+					{
+						ralloc_asprintf_append(buffer, "alignas(4) ", "");
+					}
+					
 					print_type_pre(s->fields.structure[j].type);
 					ralloc_asprintf_append(buffer, " %s", s->fields.structure[j].name);
 					if (!s->fields.structure[j].semantic || strncmp(s->fields.structure[j].semantic, "[[", 2))
