Add separate Metal shader platforms for tvOS - we were relying upon a bug in Apple's metal toolchain to allow sharing shader bytecode between iOS and tvOS but that isn't possible anymore.
- Added IsMetalMobilePlatform, IsMetalMRTPlatform & IsMetalSM5Platform to simplify testing shader platforms.
- Twiddled the shader compiler to emit the correct min-OS flags to compile for iOS and tvOS.

diff --git a/Engine/Source/Developer/Apple/MetalShaderFormat/Private/MetalShaderCompiler.cpp b/Engine/Source/Developer/Apple/MetalShaderFormat/Private/MetalShaderCompiler.cpp
--- a/Engine/Source/Developer/Apple/MetalShaderFormat/Private/MetalShaderCompiler.cpp
+++ b/Engine/Source/Developer/Apple/MetalShaderFormat/Private/MetalShaderCompiler.cpp
@@ -468,7 +468,7 @@ bool CopyRemoteFileToLocal(FString const& RemotePath, FString const& LocalPath)
 
 FString GetMetalBinaryPath(uint32 ShaderPlatform)
 {
-	const bool bIsMobile = (ShaderPlatform == SP_METAL || ShaderPlatform == SP_METAL_MRT);
+	const bool bIsMobile = (ShaderPlatform == SP_METAL || ShaderPlatform == SP_METAL_MRT || ShaderPlatform == SP_METAL_TVOS || ShaderPlatform == SP_METAL_MRT_TVOS);
 	if(GMetalBinaryPath[bIsMobile].Len() == 0 || GMetalToolsPath[bIsMobile].Len() == 0)
 	{
 		FString XcodePath = GetXcodePath();
@@ -514,7 +514,7 @@ FString GetMetalToolsPath(uint32 ShaderPlatform)
 {
 	GetMetalBinaryPath(ShaderPlatform);
 	
-	const bool bIsMobile = (ShaderPlatform == SP_METAL || ShaderPlatform == SP_METAL_MRT);
+	const bool bIsMobile = (ShaderPlatform == SP_METAL || ShaderPlatform == SP_METAL_MRT || ShaderPlatform == SP_METAL_TVOS || ShaderPlatform == SP_METAL_MRT_TVOS);
 	return GMetalToolsPath[bIsMobile];
 }
 
@@ -522,7 +522,7 @@ FString GetMetalLibraryPath(uint32 ShaderPlatform)
 {
 	GetMetalBinaryPath(ShaderPlatform);
 	
-	const bool bIsMobile = (ShaderPlatform == SP_METAL || ShaderPlatform == SP_METAL_MRT);
+	const bool bIsMobile = (ShaderPlatform == SP_METAL || ShaderPlatform == SP_METAL_MRT || ShaderPlatform == SP_METAL_TVOS || ShaderPlatform == SP_METAL_MRT_TVOS);
 	return GMetalLibraryPath[bIsMobile];
 }
 
@@ -530,7 +530,7 @@ FString GetMetalCompilerVersion(uint32 ShaderPlatform)
 {
 	GetMetalBinaryPath(ShaderPlatform);
 	
-	const bool bIsMobile = (ShaderPlatform == SP_METAL || ShaderPlatform == SP_METAL_MRT);
+	const bool bIsMobile = (ShaderPlatform == SP_METAL || ShaderPlatform == SP_METAL_MRT || ShaderPlatform == SP_METAL_TVOS || ShaderPlatform == SP_METAL_MRT_TVOS);
 	return GMetalCompilerVers[bIsMobile];
 }
 
@@ -971,10 +971,10 @@ void BuildMetalShaderOutput(
 	const ANSICHAR* SideTableString = FCStringAnsi::Strstr(USFSource, "@SideTable: ");
 	
 	EShaderFrequency Frequency = (EShaderFrequency)ShaderOutput.Target.Frequency;
-	const bool bIsMobile = (ShaderInput.Target.Platform == SP_METAL || ShaderInput.Target.Platform == SP_METAL_MRT);
+	const bool bIsMobile = (ShaderInput.Target.Platform == SP_METAL || ShaderInput.Target.Platform == SP_METAL_MRT || ShaderInput.Target.Platform == SP_METAL_TVOS || ShaderInput.Target.Platform == SP_METAL_MRT_TVOS);
 	bool bNoFastMath = ShaderInput.Environment.CompilerFlags.Contains(CFLAG_NoFastMath);
 	FString const* UsingWPO = ShaderInput.Environment.GetDefinitions().Find(TEXT("USES_WORLD_POSITION_OFFSET"));
-	if (UsingWPO && FString("1") == *UsingWPO && (ShaderInput.Target.Platform == SP_METAL_MRT) && Frequency == SF_Vertex)
+	if (UsingWPO && FString("1") == *UsingWPO && (ShaderInput.Target.Platform == SP_METAL_MRT || ShaderInput.Target.Platform == SP_METAL_MRT_TVOS) && Frequency == SF_Vertex)
 	{
 		// WPO requires that we make all multiply/sincos instructions invariant :(
 		bNoFastMath = true;
@@ -1963,7 +1963,7 @@ void CompileShader_Metal(const FShaderCompilerInput& _Input,FShaderCompilerOutpu
 	EHlslCompileTarget MetalCompilerTarget = HCT_FeatureLevelES3_1; // Varies depending on the actual intended Metal target.
 
 	// Work out which standard we need, this is dependent on the shader platform.
-	const bool bIsMobile = (Input.Target.Platform == SP_METAL || Input.Target.Platform == SP_METAL_MRT);
+	const bool bIsMobile = (Input.Target.Platform == SP_METAL || Input.Target.Platform == SP_METAL_MRT || Input.Target.Platform == SP_METAL_TVOS || Input.Target.Platform == SP_METAL_MRT_TVOS);
 	TCHAR const* StandardPlatform = nullptr;
 	if (bIsMobile)
 	{
@@ -1982,6 +1982,8 @@ void CompileShader_Metal(const FShaderCompilerInput& _Input,FShaderCompilerOutpu
 
 	static FName NAME_SF_METAL(TEXT("SF_METAL"));
 	static FName NAME_SF_METAL_MRT(TEXT("SF_METAL_MRT"));
+	static FName NAME_SF_METAL_TVOS(TEXT("SF_METAL_TVOS"));
+	static FName NAME_SF_METAL_MRT_TVOS(TEXT("SF_METAL_MRT_TVOS"));
 	static FName NAME_SF_METAL_SM5_NOTESS(TEXT("SF_METAL_SM5_NOTESS"));
 	static FName NAME_SF_METAL_SM5(TEXT("SF_METAL_SM5"));
 	static FName NAME_SF_METAL_MACES3_1(TEXT("SF_METAL_MACES3_1"));
@@ -1999,13 +2001,14 @@ void CompileShader_Metal(const FShaderCompilerInput& _Input,FShaderCompilerOutpu
             LexFromString(VersionEnum, *(*MaxVersion));
         }
     }
-    
-    if (Input.ShaderFormat == NAME_SF_METAL)
+	
+	bool bAppleTV = (Input.ShaderFormat == NAME_SF_METAL_TVOS || Input.ShaderFormat == NAME_SF_METAL_MRT_TVOS);
+    if (Input.ShaderFormat == NAME_SF_METAL || Input.ShaderFormat == NAME_SF_METAL_TVOS)
 	{
         VersionEnum = VersionEnum > 0 ? VersionEnum : 0;
         AdditionalDefines.SetDefine(TEXT("METAL_PROFILE"), 1);
 	}
-	else if (Input.ShaderFormat == NAME_SF_METAL_MRT)
+	else if (Input.ShaderFormat == NAME_SF_METAL_MRT || Input.ShaderFormat == NAME_SF_METAL_MRT_TVOS)
 	{
         UE_CLOG(VersionEnum == 0, LogShaders, Warning, TEXT("Metal shader version should be Metal v1.2 or higher for format %s!"), VersionEnum, *Input.ShaderFormat.ToString());
 		AdditionalDefines.SetDefine(TEXT("METAL_MRT_PROFILE"), 1);
@@ -2069,24 +2072,57 @@ void CompileShader_Metal(const FShaderCompilerInput& _Input,FShaderCompilerOutpu
 	{
 		case 4:
 			// Enable full SM5 feature support so tessellation & fragment UAVs compile
+			TypeMode = EMetalTypeBufferModeTex;
             HlslCompilerTarget = HCT_FeatureLevelSM5;
 			StandardVersion = TEXT("2.1");
-			MinOSVersion = bIsMobile ? TEXT("") : TEXT("-mmacosx-version-min=10.14");
-			TypeMode = EMetalTypeBufferModeTex;
+			if (bAppleTV)
+			{
+				MinOSVersion = TEXT("-mtvos-version-min=12.0");
+			}
+			else if (bIsMobile)
+			{
+				MinOSVersion = TEXT("-mios-version-min=12.0");
+			}
+			else
+			{
+				MinOSVersion = TEXT("-mmacosx-version-min=10.14");
+			}
 			break;
 		case 3:
 			// Enable full SM5 feature support so tessellation & fragment UAVs compile
+			TypeMode = EMetalTypeBufferModeUAV;
             HlslCompilerTarget = HCT_FeatureLevelSM5;
 			StandardVersion = TEXT("2.0");
-			MinOSVersion = bIsMobile ? TEXT("") : TEXT("-mmacosx-version-min=10.13");
-			TypeMode = EMetalTypeBufferModeUAV;
+			if (bAppleTV)
+			{
+				MinOSVersion = TEXT("-mtvos-version-min=11.0");
+			}
+			else if (bIsMobile)
+			{
+				MinOSVersion = TEXT("-mios-version-min=11.0");
+			}
+			else
+			{
+				MinOSVersion = TEXT("-mmacosx-version-min=10.13");
+			}
 			break;
 		case 2:
 			// Enable full SM5 feature support so tessellation & fragment UAVs compile
+			TypeMode = EMetalTypeBufferModeSRV;
             HlslCompilerTarget = HCT_FeatureLevelSM5;
 			StandardVersion = TEXT("1.2");
-			MinOSVersion = bIsMobile ? TEXT("") : TEXT("-mmacosx-version-min=10.12");
-			TypeMode = EMetalTypeBufferModeSRV;
+			if (bAppleTV)
+			{
+				MinOSVersion = TEXT("-mtvos-version-min=10.0");
+			}
+			else if (bIsMobile)
+			{
+				MinOSVersion = TEXT("-mios-version-min=10.0");
+			}
+			else
+			{
+				MinOSVersion = TEXT("-mmacosx-version-min=10.12");
+			}
 			break;
 		case 1:
 			HlslCompilerTarget = bIsMobile ? HlslCompilerTarget : HCT_FeatureLevelSM5;
@@ -2471,6 +2507,8 @@ EShaderPlatform MetalShaderFormatToLegacyShaderPlatform(FName ShaderFormat)
 {
 	static FName NAME_SF_METAL(TEXT("SF_METAL"));
 	static FName NAME_SF_METAL_MRT(TEXT("SF_METAL_MRT"));
+	static FName NAME_SF_METAL_TVOS(TEXT("SF_METAL_TVOS"));
+	static FName NAME_SF_METAL_MRT_TVOS(TEXT("SF_METAL_MRT_TVOS"));
 	static FName NAME_SF_METAL_SM5_NOTESS(TEXT("SF_METAL_SM5_NOTESS"));
 	static FName NAME_SF_METAL_SM5(TEXT("SF_METAL_SM5"));
 	static FName NAME_SF_METAL_MRT_MAC(TEXT("SF_METAL_MRT_MAC"));
@@ -2479,6 +2517,8 @@ EShaderPlatform MetalShaderFormatToLegacyShaderPlatform(FName ShaderFormat)
 	
 	if (ShaderFormat == NAME_SF_METAL)				return SP_METAL;
 	if (ShaderFormat == NAME_SF_METAL_MRT)			return SP_METAL_MRT;
+	if (ShaderFormat == NAME_SF_METAL_TVOS)			return SP_METAL_TVOS;
+	if (ShaderFormat == NAME_SF_METAL_MRT_TVOS)		return SP_METAL_MRT_TVOS;
 	if (ShaderFormat == NAME_SF_METAL_MRT_MAC)		return SP_METAL_MRT_MAC;
 	if (ShaderFormat == NAME_SF_METAL_SM5)			return SP_METAL_SM5;
 	if (ShaderFormat == NAME_SF_METAL_SM5_NOTESS)	return SP_METAL_SM5_NOTESS;
diff --git a/Engine/Source/Developer/Apple/MetalShaderFormat/Private/MetalShaderFormat.cpp b/Engine/Source/Developer/Apple/MetalShaderFormat/Private/MetalShaderFormat.cpp
--- a/Engine/Source/Developer/Apple/MetalShaderFormat/Private/MetalShaderFormat.cpp
+++ b/Engine/Source/Developer/Apple/MetalShaderFormat/Private/MetalShaderFormat.cpp
@@ -21,6 +21,8 @@ extern bool FinalizeLibrary_Metal(class FName const& Format, class FString const
 
 static FName NAME_SF_METAL(TEXT("SF_METAL"));
 static FName NAME_SF_METAL_MRT(TEXT("SF_METAL_MRT"));
+static FName NAME_SF_METAL_TVOS(TEXT("SF_METAL_TVOS"));
+static FName NAME_SF_METAL_MRT_TVOS(TEXT("SF_METAL_MRT_TVOS"));
 static FName NAME_SF_METAL_SM5_NOTESS(TEXT("SF_METAL_SM5_NOTESS"));
 static FName NAME_SF_METAL_SM5(TEXT("SF_METAL_SM5"));
 static FName NAME_SF_METAL_MACES3_1(TEXT("SF_METAL_MACES3_1"));
@@ -38,7 +40,7 @@ public:
 	, WorkingDir(WorkingDirectory)
 	{
 		check(LibraryName.Len() > 0);
-		check(Format == NAME_SF_METAL || Format == NAME_SF_METAL_MRT || Format == NAME_SF_METAL_SM5_NOTESS || Format == NAME_SF_METAL_SM5 || Format == NAME_SF_METAL_MACES3_1 || Format == NAME_SF_METAL_MACES2 || Format == NAME_SF_METAL_MRT_MAC);
+		check(Format == NAME_SF_METAL || Format == NAME_SF_METAL_MRT || Format == NAME_SF_METAL_TVOS || Format == NAME_SF_METAL_MRT_TVOS || Format == NAME_SF_METAL_SM5_NOTESS || Format == NAME_SF_METAL_SM5 || Format == NAME_SF_METAL_MACES3_1 || Format == NAME_SF_METAL_MACES2 || Format == NAME_SF_METAL_MRT_MAC);
 		ArchivePath = (WorkingDir / Format.GetPlainNameString());
 		IFileManager::Get().DeleteDirectory(*ArchivePath, false, true);
 		IFileManager::Get().MakeDirectory(*ArchivePath);
@@ -255,6 +257,8 @@ public:
 	{
 		OutFormats.Add(NAME_SF_METAL);
 		OutFormats.Add(NAME_SF_METAL_MRT);
+		OutFormats.Add(NAME_SF_METAL_TVOS);
+		OutFormats.Add(NAME_SF_METAL_MRT_TVOS);
 		OutFormats.Add(NAME_SF_METAL_SM5_NOTESS);
 		OutFormats.Add(NAME_SF_METAL_SM5);
 		OutFormats.Add(NAME_SF_METAL_MACES3_1);
@@ -263,7 +267,7 @@ public:
 	}
 	virtual void CompileShader(FName Format, const struct FShaderCompilerInput& Input, struct FShaderCompilerOutput& Output,const FString& WorkingDirectory) const override final
 	{
-		check(Format == NAME_SF_METAL || Format == NAME_SF_METAL_MRT || Format == NAME_SF_METAL_SM5_NOTESS || Format == NAME_SF_METAL_SM5 || Format == NAME_SF_METAL_MACES3_1 || Format == NAME_SF_METAL_MACES2 || Format == NAME_SF_METAL_MRT_MAC);
+		check(Format == NAME_SF_METAL || Format == NAME_SF_METAL_MRT || Format == NAME_SF_METAL_TVOS || Format == NAME_SF_METAL_MRT_TVOS || Format == NAME_SF_METAL_SM5_NOTESS || Format == NAME_SF_METAL_SM5 || Format == NAME_SF_METAL_MACES3_1 || Format == NAME_SF_METAL_MACES2 || Format == NAME_SF_METAL_MRT_MAC);
 		CompileShader_Metal(Input, Output, WorkingDirectory);
 	}
 	virtual bool CanStripShaderCode(bool const bNativeFormat) const override final
@@ -308,7 +312,7 @@ uint32 GetMetalFormatVersion(FName Format)
 	// Include the Xcode version when the .ini settings instruct us to do so.
 	uint16 AppVersion = 0;
 	bool bAddXcodeVersionInShaderVersion = false;
-	if(Format == NAME_SF_METAL || Format == NAME_SF_METAL_MRT)
+	if(Format == NAME_SF_METAL || Format == NAME_SF_METAL_MRT || Format == NAME_SF_METAL_TVOS || Format == NAME_SF_METAL_MRT_TVOS)
 	{
 		GConfig->GetBool(TEXT("/Script/IOSRuntimeSettings.IOSRuntimeSettings"), TEXT("XcodeVersionInShaderVersion"), bAddXcodeVersionInShaderVersion, GEngineIni);
 	}
diff --git a/Engine/Source/Developer/IOS/IOSTargetPlatform/Private/IOSTargetPlatform.cpp b/Engine/Source/Developer/IOS/IOSTargetPlatform/Private/IOSTargetPlatform.cpp
--- a/Engine/Source/Developer/IOS/IOSTargetPlatform/Private/IOSTargetPlatform.cpp
+++ b/Engine/Source/Developer/IOS/IOSTargetPlatform/Private/IOSTargetPlatform.cpp
@@ -531,19 +531,21 @@ void FIOSTargetPlatform::GetAllPossibleShaderFormats( TArray<FName>& OutFormats
 	static FName NAME_GLSL_ES2_IOS(TEXT("GLSL_ES2_IOS"));
 	static FName NAME_SF_METAL(TEXT("SF_METAL"));
 	static FName NAME_SF_METAL_MRT(TEXT("SF_METAL_MRT"));
+	static FName NAME_SF_METAL_TVOS(TEXT("SF_METAL_TVOS"));
+	static FName NAME_SF_METAL_MRT_TVOS(TEXT("SF_METAL_MRT_TVOS"));
 
 	if (bIsTVOS)
 	{
 		if (SupportsMetalMRT())
 		{
-			OutFormats.AddUnique(NAME_SF_METAL_MRT);
+			OutFormats.AddUnique(NAME_SF_METAL_MRT_TVOS);
 		}
 
 		// because we are currently using IOS settings, we will always use metal, even if Metal isn't listed as being supported
 		// however, if MetalMRT is specific and Metal is set to false, then we will just use MetalMRT
 		if (SupportsMetal() || !SupportsMetalMRT())
 		{
-			OutFormats.AddUnique(NAME_SF_METAL);
+			OutFormats.AddUnique(NAME_SF_METAL_TVOS);
 		}
 	}
 	else
diff --git a/Engine/Source/Developer/SharedSettingsWidgets/Private/ShaderFormatsPropertyDetails.cpp b/Engine/Source/Developer/SharedSettingsWidgets/Private/ShaderFormatsPropertyDetails.cpp
--- a/Engine/Source/Developer/SharedSettingsWidgets/Private/ShaderFormatsPropertyDetails.cpp
+++ b/Engine/Source/Developer/SharedSettingsWidgets/Private/ShaderFormatsPropertyDetails.cpp
@@ -63,10 +63,16 @@ static FText GetFriendlyNameFromRHIName(const FString& InRHIName)
 			FriendlyRHIName = LOCTEXT("OpenGLES31", "OpenGLES 3.1 (Mobile)");
 			break;
 		case SP_METAL:
-			FriendlyRHIName = LOCTEXT("Metal", "iOS/tvOS Metal Mobile Renderer (ES3.1, Metal 1.1+, iOS/tvOS 9.0 or later)");
+			FriendlyRHIName = LOCTEXT("Metal", "iOS Metal Mobile Renderer (ES3.1, Metal 1.1+, iOS 9.0 or later)");
 			break;
 		case SP_METAL_MRT:
-			FriendlyRHIName = LOCTEXT("MetalMRT", "iOS/tvOS Metal Desktop Renderer (SM5, Metal 1.2+, iOS/tvOS 10.0 or later)");
+			FriendlyRHIName = LOCTEXT("MetalMRT", "iOS Metal Desktop Renderer (SM5, Metal 1.2+, iOS 10.0 or later)");
+			break;
+		case SP_METAL_TVOS:
+			FriendlyRHIName = LOCTEXT("MetalTV", "tvOS Metal Mobile Renderer (ES3.1, Metal 1.1+, tvOS 9.0 or later)");
+			break;
+		case SP_METAL_MRT_TVOS:
+			FriendlyRHIName = LOCTEXT("MetalMRTTV", "tvOS Metal Desktop Renderer (SM5, Metal 1.2+, tvOS 10.0 or later)");
 			break;
 		case SP_METAL_SM5_NOTESS:
 			FriendlyRHIName = LOCTEXT("MetalSM5_NOTESS", "Mac Metal Desktop Renderer without Tessellation (SM5, Metal 1.2+, macOS Sierra 10.12.6 or later)");
diff --git a/Engine/Source/Editor/MaterialEditor/Private/MaterialStats.cpp b/Engine/Source/Editor/MaterialEditor/Private/MaterialStats.cpp
--- a/Engine/Source/Editor/MaterialEditor/Private/MaterialStats.cpp
+++ b/Engine/Source/Editor/MaterialEditor/Private/MaterialStats.cpp
@@ -428,8 +428,10 @@ void FMaterialStats::BuildShaderPlatformDB()
 	AddShaderPlatform(EPlatformCategoryType::Android, SP_OPENGL_ES2_ANDROID, TEXT("Android GLES 2.0"), true, true, TEXT("Android, OpenGLES 2.0"));
 	AddShaderPlatform(EPlatformCategoryType::Android, SP_VULKAN_ES3_1_ANDROID, TEXT("Android Vulkan"), true, true, TEXT("Android, Vulkan"));
 
-	// iOS
-	AddShaderPlatform(EPlatformCategoryType::IOS, SP_METAL_SM5, TEXT("Metal SM5"), false, true, TEXT("iOS, Metal, Shader Model 5"));
+	// Apple
+	AddShaderPlatform(EPlatformCategoryType::Desktop, SP_METAL_SM5, TEXT("Metal SM5"), false, true, TEXT("macOS, Metal, Shader Model 5"));
+	AddShaderPlatform(EPlatformCategoryType::IOS, SP_METAL, TEXT("Metal"), false, true, TEXT("iOS, Metal, Mobile"));
+	AddShaderPlatform(EPlatformCategoryType::IOS, SP_METAL_MRT, TEXT("Metal MRT"), false, true, TEXT("iOS, Metal, Shader Model 5"));
 }
 
 TSharedPtr<FShaderPlatformSettings> FMaterialStats::AddShaderPlatform(const EPlatformCategoryType PlatformType, const EShaderPlatform PlatformID, const FName PlatformName,
@@ -1113,4 +1115,4 @@ void FMaterialStats::ExtractHLSLCode()
 #undef LOCTEXT_NAMESPACE
 
 /*end FMaterialStats functions*/
-/***********************************************************************************************************************/
\ No newline at end of file
+/***********************************************************************************************************************/
diff --git a/Engine/Source/Editor/MaterialEditor/Private/MaterialStatsCommon.cpp b/Engine/Source/Editor/MaterialEditor/Private/MaterialStatsCommon.cpp
--- a/Engine/Source/Editor/MaterialEditor/Private/MaterialStatsCommon.cpp
+++ b/Engine/Source/Editor/MaterialEditor/Private/MaterialStatsCommon.cpp
@@ -156,6 +156,12 @@ FString FMaterialStatsUtils::ShaderPlatformTypeName(const EShaderPlatform Platfo
 		case SP_METAL_MRT:
 			return FString("METAL_MRT");
 		break;
+		case SP_METAL_TVOS:
+			return FString("METAL_TVOS");
+		break;
+		case SP_METAL_MRT_TVOS:
+			return FString("METAL_MRT_TVOS");
+		break;
 		case SP_OPENGL_ES31_EXT:
 			return FString("OPENGL_ES31_EXT");
 		break;
@@ -266,6 +272,8 @@ bool FMaterialStatsUtils::PlatformNeedsOfflineCompiler(const EShaderPlatform Sha
 		case SP_PCD3D_ES2:
 		case SP_METAL:
 		case SP_METAL_MRT:
+		case SP_METAL_TVOS:
+		case SP_METAL_MRT_TVOS:
 		case SP_PCD3D_ES3_1:
 		case SP_METAL_SM5:
 		case SP_METAL_SM5_NOTESS:
@@ -757,4 +765,4 @@ void FMaterialStatsUtils::ExtractMatertialStatsInfo(FShaderStatsInfo& OutInfo, c
 }
 
 /*end FMaterialStatsUtils */
-/***********************************************************************************************************************/
\ No newline at end of file
+/***********************************************************************************************************************/
diff --git a/Engine/Source/Runtime/Apple/MetalRHI/Private/MetalRHI.cpp b/Engine/Source/Runtime/Apple/MetalRHI/Private/MetalRHI.cpp
--- a/Engine/Source/Runtime/Apple/MetalRHI/Private/MetalRHI.cpp
+++ b/Engine/Source/Runtime/Apple/MetalRHI/Private/MetalRHI.cpp
@@ -48,11 +48,11 @@ static void ValidateTargetedRHIFeatureLevelExists(EShaderPlatform Platform)
 		}
 	}
 #else
-	if (Platform == SP_METAL)
+	if (Platform == SP_METAL || Platform == SP_METAL_TVOS)
 	{
 		GConfig->GetBool(TEXT("/Script/IOSRuntimeSettings.IOSRuntimeSettings"), TEXT("bSupportsMetal"), bSupportsShaderPlatform, GEngineIni);
 	}
-	else if (Platform == SP_METAL_MRT)
+	else if (Platform == SP_METAL_MRT || Platform == SP_METAL_MRT_TVOS)
 	{
 		GConfig->GetBool(TEXT("/Script/IOSRuntimeSettings.IOSRuntimeSettings"), TEXT("bSupportsMetalMRT"), bSupportsShaderPlatform, GEngineIni);
 	}
@@ -153,10 +153,14 @@ FMetalDynamicRHI::FMetalDynamicRHI(ERHIFeatureLevel::Type RequestedFeatureLevel)
     // only allow GBuffers, etc on A8s (A7s are just not going to cut it)
     if (bProjectSupportsMRTs && bCanUseWideMRTs && bRequestedMetalMRT)
     {
+#if PLATFORM_TVOS
+		ValidateTargetedRHIFeatureLevelExists(SP_METAL_MRT);
+		GMaxRHIShaderPlatform = SP_METAL_MRT_TVOS;
+#else
 		ValidateTargetedRHIFeatureLevelExists(SP_METAL_MRT);
-		
-        GMaxRHIFeatureLevel = ERHIFeatureLevel::SM5;
         GMaxRHIShaderPlatform = SP_METAL_MRT;
+#endif
+		GMaxRHIFeatureLevel = ERHIFeatureLevel::SM5;
 		
 		bSupportsRHIThread = FParse::Param(FCommandLine::Get(),TEXT("rhithread"));
     }
@@ -167,10 +171,14 @@ FMetalDynamicRHI::FMetalDynamicRHI(ERHIFeatureLevel::Type RequestedFeatureLevel)
 			UE_LOG(LogMetal, Warning, TEXT("Metal MRT support requires an iOS or tvOS device with an A8 processor or later. Falling back to Metal ES 3.1."));
 		}
 		
+#if PLATFORM_TVOS
+		ValidateTargetedRHIFeatureLevelExists(SP_METAL_TVOS);
+		GMaxRHIShaderPlatform = SP_METAL_TVOS;
+#else
 		ValidateTargetedRHIFeatureLevelExists(SP_METAL);
-		
+		GMaxRHIShaderPlatform = SP_METAL;
+#endif
         GMaxRHIFeatureLevel = ERHIFeatureLevel::ES3_1;
-        GMaxRHIShaderPlatform = SP_METAL;
 	}
 		
 	FPlatformMemoryStats Stats = FPlatformMemory::GetStats();
@@ -180,8 +188,13 @@ FMetalDynamicRHI::FMetalDynamicRHI(ERHIFeatureLevel::Type RequestedFeatureLevel)
 	MemoryStats.DedicatedSystemMemory = 0;
 	MemoryStats.SharedSystemMemory = Stats.AvailablePhysical;
 	
+#if PLATFORM_TVOS
+	GShaderPlatformForFeatureLevel[ERHIFeatureLevel::ES2] = SP_METAL_TVOS;
+	GShaderPlatformForFeatureLevel[ERHIFeatureLevel::ES3_1] = SP_METAL_TVOS;
+#else
 	GShaderPlatformForFeatureLevel[ERHIFeatureLevel::ES2] = SP_METAL;
 	GShaderPlatformForFeatureLevel[ERHIFeatureLevel::ES3_1] = SP_METAL;
+#endif
 	GShaderPlatformForFeatureLevel[ERHIFeatureLevel::SM4] = (GMaxRHIFeatureLevel >= ERHIFeatureLevel::SM4) ? GMaxRHIShaderPlatform : SP_NumPlatforms;
 	GShaderPlatformForFeatureLevel[ERHIFeatureLevel::SM5] = (GMaxRHIFeatureLevel >= ERHIFeatureLevel::SM4) ? GMaxRHIShaderPlatform : SP_NumPlatforms;
 
diff --git a/Engine/Source/Runtime/Engine/Private/GPUSkinCache.cpp b/Engine/Source/Runtime/Engine/Private/GPUSkinCache.cpp
--- a/Engine/Source/Runtime/Engine/Private/GPUSkinCache.cpp
+++ b/Engine/Source/Runtime/Engine/Private/GPUSkinCache.cpp
@@ -109,7 +109,7 @@ static int32 GGPUSkinCacheFlushCounter = 0;
 
 static inline bool DoesPlatformSupportGPUSkinCache(EShaderPlatform Platform)
 {
-	return Platform == SP_PCD3D_SM5 || Platform == SP_METAL_SM5 || Platform == SP_METAL_SM5_NOTESS || Platform == SP_METAL_MRT_MAC || Platform == SP_METAL_MRT || IsVulkanSM5Platform(Platform) || Platform == SP_OPENGL_SM5;
+	return Platform == SP_PCD3D_SM5 || IsMetalSM5Platform(Platform) || IsVulkanSM5Platform(Platform) || Platform == SP_OPENGL_SM5;
 }
 
 ENGINE_API bool IsGPUSkinCacheAvailable()
diff --git a/Engine/Source/Runtime/Engine/Private/Materials/HLSLMaterialTranslator.h b/Engine/Source/Runtime/Engine/Private/Materials/HLSLMaterialTranslator.h
--- a/Engine/Source/Runtime/Engine/Private/Materials/HLSLMaterialTranslator.h
+++ b/Engine/Source/Runtime/Engine/Private/Materials/HLSLMaterialTranslator.h
@@ -1057,7 +1057,7 @@ public:
 		}
 		
 		// @todo MetalMRT: Remove this hack and implement proper atmospheric-fog solution for Metal MRT...
-		OutEnvironment.SetDefine(TEXT("MATERIAL_ATMOSPHERIC_FOG"), (InPlatform != SP_METAL_MRT && InPlatform != SP_METAL_MRT_MAC) ? bUsesAtmosphericFog : 0);
+		OutEnvironment.SetDefine(TEXT("MATERIAL_ATMOSPHERIC_FOG"), !IsMetalMRTPlatform(InPlatform) ? bUsesAtmosphericFog : 0);
 		OutEnvironment.SetDefine(TEXT("INTERPOLATE_VERTEX_COLOR"), bUsesVertexColor);
 		OutEnvironment.SetDefine(TEXT("NEEDS_PARTICLE_COLOR"), bUsesParticleColor); 
 		OutEnvironment.SetDefine(TEXT("NEEDS_PARTICLE_TRANSFORM"), bUsesParticleTransform);
diff --git a/Engine/Source/Runtime/Engine/Private/SimpleElementShaders.cpp b/Engine/Source/Runtime/Engine/Private/SimpleElementShaders.cpp
--- a/Engine/Source/Runtime/Engine/Private/SimpleElementShaders.cpp
+++ b/Engine/Source/Runtime/Engine/Private/SimpleElementShaders.cpp
@@ -35,7 +35,7 @@ bool FSimpleElementVS::Serialize(FArchive& Ar)
 void FSimpleElementVS::ModifyCompilationEnvironment(const FGlobalShaderPermutationParameters& Parameters, FShaderCompilerEnvironment& OutEnvironment)
 {
 	FGlobalShader::ModifyCompilationEnvironment(Parameters, OutEnvironment);
-	OutEnvironment.SetDefine(TEXT("ALLOW_SWITCH_VERTICALAXIS"), (Parameters.Platform != SP_METAL));
+	OutEnvironment.SetDefine(TEXT("ALLOW_SWITCH_VERTICALAXIS"), !IsMetalMobilePlatform(Parameters.Platform));
 }
 
 /*------------------------------------------------------------------------------
diff --git a/Engine/Source/Runtime/RHI/Private/Apple/AppleDynamicRHI.cpp b/Engine/Source/Runtime/RHI/Private/Apple/AppleDynamicRHI.cpp
--- a/Engine/Source/Runtime/RHI/Private/Apple/AppleDynamicRHI.cpp
+++ b/Engine/Source/Runtime/RHI/Private/Apple/AppleDynamicRHI.cpp
@@ -55,14 +55,22 @@ FDynamicRHI* PlatformCreateDynamicRHI()
 		GConfig->GetBool(TEXT("/Script/IOSRuntimeSettings.IOSRuntimeSettings"), TEXT("bSupportsMetalMRT"), bSupportsMetalMRT, GEngineIni);
 		if (bSupportsMetalMRT)
 		{
+#if PLATFORM_TVOS
+			TargetedShaderFormats.Add(LegacyShaderPlatformToShaderFormat(SP_METAL_MRT_TVOS).ToString());
+#else
 			TargetedShaderFormats.Add(LegacyShaderPlatformToShaderFormat(SP_METAL_MRT).ToString());
+#endif
 		}
 		
 		bool bSupportsMetal = false;
 		GConfig->GetBool(TEXT("/Script/IOSRuntimeSettings.IOSRuntimeSettings"), TEXT("bSupportsMetal"), bSupportsMetal, GEngineIni);
 		if (bSupportsMetal)
 		{
+#if PLATFORM_TVOS
+			TargetedShaderFormats.Add(LegacyShaderPlatformToShaderFormat(SP_METAL_TVOS).ToString());
+#else
 			TargetedShaderFormats.Add(LegacyShaderPlatformToShaderFormat(SP_METAL).ToString());
+#endif
 		}
 	
 		bool bSupportsOpenGLES2 = false;
diff --git a/Engine/Source/Runtime/RHI/Private/RHI.cpp b/Engine/Source/Runtime/RHI/Private/RHI.cpp
--- a/Engine/Source/Runtime/RHI/Private/RHI.cpp
+++ b/Engine/Source/Runtime/RHI/Private/RHI.cpp
@@ -668,6 +668,7 @@ static FName NAME_PLATFORM_ANDROID(TEXT("Android"));
 static FName NAME_PLATFORM_IOS(TEXT("IOS"));
 static FName NAME_PLATFORM_MAC(TEXT("Mac"));
 static FName NAME_PLATFORM_SWITCH(TEXT("Switch"));
+static FName NAME_PLATFORM_TVOS(TEXT("TVOS"));
 
 FName ShaderPlatformToPlatformName(EShaderPlatform Platform)
 {
@@ -685,7 +686,10 @@ FName ShaderPlatformToPlatformName(EShaderPlatform Platform)
 		return NAME_PLATFORM_ANDROID;
 	case SP_METAL:
 	case SP_METAL_MRT:
-		return NAME_PLATFORM_IOS;
+        return NAME_PLATFORM_IOS;
+	case SP_METAL_TVOS:
+	case SP_METAL_MRT_TVOS:
+		return NAME_PLATFORM_TVOS;
 	case SP_METAL_SM5:
 	case SP_METAL_SM5_NOTESS:
 	case SP_METAL_MACES3_1:
@@ -802,7 +806,7 @@ RHI_API bool RHISupportsPixelShaderUAVs(const EShaderPlatform Platform)
 	{
 		return true;
 	}
-	else if (Platform == SP_METAL_SM5 || Platform == SP_METAL_SM5_NOTESS || Platform == SP_METAL_MRT || Platform == SP_METAL_MRT_MAC)
+	else if (IsMetalSM5Platform(Platform))
 	{
 		return (RHIGetShaderLanguageVersion(Platform) >= 2);
 	}
diff --git a/Engine/Source/Runtime/RHI/Public/RHIDefinitions.h b/Engine/Source/Runtime/RHI/Public/RHIDefinitions.h
--- a/Engine/Source/Runtime/RHI/Public/RHIDefinitions.h
+++ b/Engine/Source/Runtime/RHI/Public/RHIDefinitions.h
@@ -63,9 +63,11 @@ enum EShaderPlatform
 	SP_METAL_MRT_MAC		= 27,
 	SP_VULKAN_SM5_LUMIN		= 28,
 	SP_VULKAN_ES3_1_LUMIN	= 29,
+	SP_METAL_TVOS			= 30,
+	SP_METAL_MRT_TVOS		= 31,
 
-	SP_NumPlatforms			= 30,
-	SP_NumBits				= 5,
+	SP_NumPlatforms			= 32,
+	SP_NumBits				= 6,
 };
 static_assert(SP_NumPlatforms <= (1 << SP_NumBits), "SP_NumPlatforms will not fit on SP_NumBits");
 
@@ -811,7 +813,7 @@ inline bool IsES2Platform(const EShaderPlatform Platform)
 inline bool IsMobilePlatform(const EShaderPlatform Platform)
 {
 	return IsES2Platform(Platform)
-		|| Platform == SP_METAL || Platform == SP_METAL_MACES3_1
+		|| Platform == SP_METAL || Platform == SP_METAL_MACES3_1 || Platform == SP_METAL_TVOS
 		|| Platform == SP_PCD3D_ES3_1
 		|| Platform == SP_OPENGL_PCES3_1 || Platform == SP_OPENGL_ES3_1_ANDROID
 		|| Platform == SP_VULKAN_ES3_1_ANDROID || Platform == SP_VULKAN_PCES3_1 || Platform == SP_VULKAN_ES3_1_LUMIN
@@ -827,7 +829,22 @@ inline bool IsOpenGLPlatform(const EShaderPlatform Platform)
 
 inline bool IsMetalPlatform(const EShaderPlatform Platform)
 {
-	return Platform == SP_METAL || Platform == SP_METAL_MRT || Platform == SP_METAL_SM5_NOTESS || Platform == SP_METAL_SM5 || Platform == SP_METAL_MACES3_1 || Platform == SP_METAL_MACES2 || Platform == SP_METAL_MRT_MAC;
+	return Platform == SP_METAL || Platform == SP_METAL_MRT || Platform == SP_METAL_TVOS || Platform == SP_METAL_MRT_TVOS || Platform == SP_METAL_SM5_NOTESS || Platform == SP_METAL_SM5 || Platform == SP_METAL_MACES3_1 || Platform == SP_METAL_MACES2 || Platform == SP_METAL_MRT_MAC;
+}
+
+inline bool IsMetalMobilePlatform(const EShaderPlatform Platform)
+{
+	return Platform == SP_METAL || Platform == SP_METAL_TVOS;
+}
+
+inline bool IsMetalMRTPlatform(const EShaderPlatform Platform)
+{
+	return Platform == SP_METAL_MRT || Platform == SP_METAL_MRT_TVOS || Platform == SP_METAL_MRT_MAC;
+}
+
+inline bool IsMetalSM5Platform(const EShaderPlatform Platform)
+{
+	return Platform == SP_METAL_MRT || Platform == SP_METAL_MRT_TVOS || Platform == SP_METAL_SM5_NOTESS || Platform == SP_METAL_SM5 || Platform == SP_METAL_MRT_MAC;
 }
 
 inline bool IsConsolePlatform(const EShaderPlatform Platform)
@@ -899,6 +916,7 @@ inline ERHIFeatureLevel::Type GetMaxSupportedFeatureLevel(EShaderPlatform InShad
 	case SP_OPENGL_ES31_EXT:
 	case SP_METAL_SM5:
 	case SP_METAL_MRT:
+	case SP_METAL_MRT_TVOS:
 	case SP_METAL_MRT_MAC:
 	case SP_METAL_SM5_NOTESS:
 	case SP_VULKAN_SM5:
@@ -917,6 +935,7 @@ inline ERHIFeatureLevel::Type GetMaxSupportedFeatureLevel(EShaderPlatform InShad
 	case SP_METAL_MACES2:
 		return ERHIFeatureLevel::ES2;
 	case SP_METAL:
+	case SP_METAL_TVOS:
 	case SP_METAL_MACES3_1:
 	case SP_PCD3D_ES3_1:
 	case SP_OPENGL_PCES3_1:
@@ -992,14 +1011,14 @@ inline bool RHINeedsToSwitchVerticalAxis(EShaderPlatform Platform)
 #endif
 
 	// ES2 & ES3.1 need to flip when rendering to an RT that will be post processed
-	return IsOpenGLPlatform(Platform) && IsMobilePlatform(Platform) && !IsPCPlatform(Platform) && Platform != SP_METAL && !IsVulkanPlatform(Platform)
+	return IsOpenGLPlatform(Platform) && IsMobilePlatform(Platform) && !IsPCPlatform(Platform) && !IsMetalMobilePlatform(Platform) && !IsVulkanPlatform(Platform)
 		   && Platform != SP_SWITCH && Platform != SP_SWITCH_FORWARD;
 }
 
 inline bool RHISupportsSeparateMSAAAndResolveTextures(const EShaderPlatform Platform)
 {
 	// Metal mobile devices, Vulkan and Android ES2/3.1 need to handle MSAA and resolve textures internally (unless RHICreateTexture2D was changed to take an optional resolve target)
-	const bool bMobileMetalDevice = (Platform == SP_METAL || Platform == SP_METAL_MRT);
+	const bool bMobileMetalDevice = (Platform == SP_METAL || Platform == SP_METAL_TVOS || Platform == SP_METAL_MRT || Platform == SP_METAL_MRT_TVOS);
 	return !bMobileMetalDevice && !IsVulkanPlatform(Platform) && !IsAndroidOpenGLESPlatform(Platform);
 }
 
@@ -1017,7 +1036,7 @@ inline bool RHISupportsGeometryShaders(const EShaderPlatform Platform)
 inline bool RHIHasTiledGPU(const EShaderPlatform Platform)
 {
 	// @todo MetalMRT Technically we should include (Platform == SP_METAL_MRT) but this would disable depth-pre-pass which is currently required.
-	return Platform == SP_METAL || Platform == SP_OPENGL_ES2_IOS || Platform == SP_OPENGL_ES2_ANDROID || Platform == SP_OPENGL_ES3_1_ANDROID;
+	return Platform == SP_METAL || Platform == SP_METAL_TVOS || Platform == SP_OPENGL_ES2_IOS || Platform == SP_OPENGL_ES2_ANDROID || Platform == SP_OPENGL_ES3_1_ANDROID;
 }
 
 inline bool RHISupportsVertexShaderLayer(const EShaderPlatform Platform)
diff --git a/Engine/Source/Runtime/RHI/Public/RHIShaderFormatDefinitions.inl b/Engine/Source/Runtime/RHI/Public/RHIShaderFormatDefinitions.inl
--- a/Engine/Source/Runtime/RHI/Public/RHIShaderFormatDefinitions.inl
+++ b/Engine/Source/Runtime/RHI/Public/RHIShaderFormatDefinitions.inl
@@ -33,6 +33,8 @@ static FName NAME_GLSL_SWITCH_FORWARD(TEXT("GLSL_SWITCH_FORWARD"));
 
 static FName NAME_SF_METAL(TEXT("SF_METAL"));
 static FName NAME_SF_METAL_MRT(TEXT("SF_METAL_MRT"));
+static FName NAME_SF_METAL_TVOS(TEXT("SF_METAL_TVOS"));
+static FName NAME_SF_METAL_MRT_TVOS(TEXT("SF_METAL_MRT_TVOS"));
 static FName NAME_SF_METAL_MRT_MAC(TEXT("SF_METAL_MRT_MAC"));
 static FName NAME_SF_METAL_SM5(TEXT("SF_METAL_SM5"));
 static FName NAME_SF_METAL_SM5_NOTESS(TEXT("SF_METAL_SM5_NOTESS"));
@@ -101,6 +103,10 @@ static FName ShaderPlatformToShaderFormatName(EShaderPlatform Platform)
 		return NAME_SF_METAL;
 	case SP_METAL_MRT:
 		return NAME_SF_METAL_MRT;
+	case SP_METAL_TVOS:
+		return NAME_SF_METAL_TVOS;
+	case SP_METAL_MRT_TVOS:
+		return NAME_SF_METAL_MRT_TVOS;
 	case SP_METAL_MRT_MAC:
 		return NAME_SF_METAL_MRT_MAC;
 	case SP_METAL_SM5:
@@ -183,6 +189,8 @@ static EShaderPlatform ShaderFormatNameToShaderPlatform(FName ShaderFormat)
 
 	if (ShaderFormat == NAME_SF_METAL)					return SP_METAL;
 	if (ShaderFormat == NAME_SF_METAL_MRT)				return SP_METAL_MRT;
+	if (ShaderFormat == NAME_SF_METAL_TVOS)				return SP_METAL_TVOS;
+	if (ShaderFormat == NAME_SF_METAL_MRT_TVOS)			return SP_METAL_MRT_TVOS;
 	if (ShaderFormat == NAME_SF_METAL_MRT_MAC)			return SP_METAL_MRT_MAC;
 	if (ShaderFormat == NAME_SF_METAL_SM5)				return SP_METAL_SM5;
 	if (ShaderFormat == NAME_SF_METAL_SM5_NOTESS)		return SP_METAL_SM5_NOTESS;
diff --git a/Engine/Source/Runtime/Renderer/Private/BasePassRendering.h b/Engine/Source/Runtime/Renderer/Private/BasePassRendering.h
--- a/Engine/Source/Runtime/Renderer/Private/BasePassRendering.h
+++ b/Engine/Source/Runtime/Renderer/Private/BasePassRendering.h
@@ -313,7 +313,7 @@ public:
 	{
 		Super::ModifyCompilationEnvironment(Platform, Material, OutEnvironment);
 		// @todo MetalMRT: Remove this hack and implement proper atmospheric-fog solution for Metal MRT...
-		OutEnvironment.SetDefine(TEXT("BASEPASS_ATMOSPHERIC_FOG"), (Platform != SP_METAL_MRT && Platform != SP_METAL_MRT_MAC) ? bEnableAtmosphericFog : 0);
+		OutEnvironment.SetDefine(TEXT("BASEPASS_ATMOSPHERIC_FOG"), !IsMetalMRTPlatform(Platform) ? bEnableAtmosphericFog : 0);
 	}
 };
 
diff --git a/Engine/Source/Runtime/Renderer/Private/CapsuleShadowRendering.h b/Engine/Source/Runtime/Renderer/Private/CapsuleShadowRendering.h
--- a/Engine/Source/Runtime/Renderer/Private/CapsuleShadowRendering.h
+++ b/Engine/Source/Runtime/Renderer/Private/CapsuleShadowRendering.h
@@ -17,7 +17,7 @@ inline bool DoesPlatformSupportCapsuleShadows(EShaderPlatform Platform)
 {
 	// Hasn't been tested elsewhere yet
 	return Platform == SP_PCD3D_SM5 || Platform == SP_PS4 || Platform == SP_XBOXONE_D3D12
-		|| Platform == SP_METAL_SM5 || Platform == SP_METAL_SM5_NOTESS || Platform == SP_METAL_MRT_MAC || Platform == SP_METAL_MRT
+		|| IsMetalSM5Platform(Platform)
 		|| IsVulkanSM5Platform(Platform)
 		|| Platform == SP_SWITCH;
 }
@@ -37,4 +37,4 @@ inline bool SupportsCapsuleDirectShadows(ERHIFeatureLevel::Type FeatureLevel, ES
 inline bool SupportsCapsuleIndirectShadows(ERHIFeatureLevel::Type FeatureLevel, EShaderPlatform ShaderPlatform)
 {
 	return GCapsuleIndirectShadows && SupportsCapsuleShadows(FeatureLevel, ShaderPlatform);
-}
\ No newline at end of file
+}
diff --git a/Engine/Source/Runtime/Renderer/Private/GPUFastFourierTransform.cpp b/Engine/Source/Runtime/Renderer/Private/GPUFastFourierTransform.cpp
--- a/Engine/Source/Runtime/Renderer/Private/GPUFastFourierTransform.cpp
+++ b/Engine/Source/Runtime/Renderer/Private/GPUFastFourierTransform.cpp
@@ -200,7 +200,7 @@ namespace GPUFFT
 		static bool ShouldCompilePermutation(const FGlobalShaderPermutationParameters& Parameters)
 		{
 			// @todo MetalMRT: Metal MRT can't cope with the threadgroup storage requirements for these shaders right now
-			return IsFeatureLevelSupported(Parameters.Platform, ERHIFeatureLevel::SM5) && (!IsMetalPlatform(Parameters.Platform) || RHIGetShaderLanguageVersion(Parameters.Platform) >= 2) && (Parameters.Platform != SP_METAL_MRT);
+			return IsFeatureLevelSupported(Parameters.Platform, ERHIFeatureLevel::SM5) && (!IsMetalPlatform(Parameters.Platform) || RHIGetShaderLanguageVersion(Parameters.Platform) >= 2) && !IsMetalMRTPlatform(Parameters.Platform);
 		}
 
 		static void ModifyCompilationEnvironment(const FGlobalShaderPermutationParameters& Parameters, FShaderCompilerEnvironment& OutEnvironment)
@@ -306,7 +306,7 @@ namespace GPUFFT
 		static bool ShouldCompilePermutation(const FGlobalShaderPermutationParameters& Parameters)
 		{
 			// @todo MetalMRT: Metal MRT can't cope with the threadgroup storage requirements for these shaders right now
-			return IsFeatureLevelSupported(Parameters.Platform, ERHIFeatureLevel::SM5) && (!IsMetalPlatform(Parameters.Platform) || RHIGetShaderLanguageVersion(Parameters.Platform) >= 2) && (Parameters.Platform != SP_METAL_MRT);
+			return IsFeatureLevelSupported(Parameters.Platform, ERHIFeatureLevel::SM5) && (!IsMetalPlatform(Parameters.Platform) || RHIGetShaderLanguageVersion(Parameters.Platform) >= 2) && !IsMetalMRTPlatform(Parameters.Platform);
 		}
 
 		static void ModifyCompilationEnvironment(const FGlobalShaderPermutationParameters& Parameters, FShaderCompilerEnvironment& OutEnvironment)
@@ -408,7 +408,7 @@ namespace GPUFFT
 		static bool ShouldCompilePermutation(const FGlobalShaderPermutationParameters& Parameters)
 		{
 			// @todo MetalMRT: Metal MRT can't cope with the threadgroup storage requirements for these shaders right now
-			return IsFeatureLevelSupported(Parameters.Platform, ERHIFeatureLevel::SM5) && (!IsMetalPlatform(Parameters.Platform) || RHIGetShaderLanguageVersion(Parameters.Platform) >= 2) && (Parameters.Platform != SP_METAL_MRT);
+			return IsFeatureLevelSupported(Parameters.Platform, ERHIFeatureLevel::SM5) && (!IsMetalPlatform(Parameters.Platform) || RHIGetShaderLanguageVersion(Parameters.Platform) >= 2) && !IsMetalMRTPlatform(Parameters.Platform);
 		}
 
 		static void ModifyCompilationEnvironment(const FGlobalShaderPermutationParameters& Parameters, FShaderCompilerEnvironment& OutEnvironment)
@@ -509,7 +509,7 @@ namespace GPUFFT
 		static bool ShouldCompilePermutation(const FGlobalShaderPermutationParameters& Parameters)
 		{
 			// @todo MetalMRT: Metal MRT can't cope with the threadgroup storage requirements for these shaders right now
-			return IsFeatureLevelSupported(Parameters.Platform, ERHIFeatureLevel::SM5) && (!IsMetalPlatform(Parameters.Platform) || RHIGetShaderLanguageVersion(Parameters.Platform) >= 2) && (Parameters.Platform != SP_METAL_MRT);
+			return IsFeatureLevelSupported(Parameters.Platform, ERHIFeatureLevel::SM5) && (!IsMetalPlatform(Parameters.Platform) || RHIGetShaderLanguageVersion(Parameters.Platform) >= 2) && !IsMetalMRTPlatform(Parameters.Platform);
 		}
 
 		static void ModifyCompilationEnvironment(const FGlobalShaderPermutationParameters& Parameters, FShaderCompilerEnvironment& OutEnvironment)
@@ -595,7 +595,7 @@ namespace GPUFFT
 		static bool ShouldCompilePermutation(const FGlobalShaderPermutationParameters& Parameters)
 		{
 			// @todo MetalMRT: Metal MRT can't cope with the threadgroup storage requirements for these shaders right now
-			return IsFeatureLevelSupported(Parameters.Platform, ERHIFeatureLevel::SM5) && (!IsMetalPlatform(Parameters.Platform) || RHIGetShaderLanguageVersion(Parameters.Platform) >= 2) && (Parameters.Platform != SP_METAL_MRT);
+			return IsFeatureLevelSupported(Parameters.Platform, ERHIFeatureLevel::SM5) && (!IsMetalPlatform(Parameters.Platform) || RHIGetShaderLanguageVersion(Parameters.Platform) >= 2) && !IsMetalMRTPlatform(Parameters.Platform);
 		}
 
 		static void ModifyCompilationEnvironment(const FGlobalShaderPermutationParameters& Parameters, FShaderCompilerEnvironment& OutEnvironment)
@@ -680,7 +680,7 @@ namespace GPUFFT
 		static bool ShouldCompilePermutation(const FGlobalShaderPermutationParameters& Parameters)
 		{
 			// @todo MetalMRT: Metal MRT can't cope with the threadgroup storage requirements for these shaders right now
-			return IsFeatureLevelSupported(Parameters.Platform, ERHIFeatureLevel::SM5) && (!IsMetalPlatform(Parameters.Platform) || RHIGetShaderLanguageVersion(Parameters.Platform) >= 2) && (Parameters.Platform != SP_METAL_MRT);
+			return IsFeatureLevelSupported(Parameters.Platform, ERHIFeatureLevel::SM5) && (!IsMetalPlatform(Parameters.Platform) || RHIGetShaderLanguageVersion(Parameters.Platform) >= 2) && !IsMetalMRTPlatform(Parameters.Platform);
 		}
 
 		static void ModifyCompilationEnvironment(const FGlobalShaderPermutationParameters& Parameters, FShaderCompilerEnvironment& OutEnvironment)
@@ -839,7 +839,7 @@ namespace GPUFFT
 		static bool ShouldCompilePermutation(const FGlobalShaderPermutationParameters& Parameters)
 		{
 			// @todo MetalMRT: Metal MRT can't cope with the threadgroup storage requirements for these shaders right now
-			return IsFeatureLevelSupported(Parameters.Platform, ERHIFeatureLevel::SM5) && (!IsMetalPlatform(Parameters.Platform) || RHIGetShaderLanguageVersion(Parameters.Platform) >= 2) && (Parameters.Platform != SP_METAL_MRT);
+			return IsFeatureLevelSupported(Parameters.Platform, ERHIFeatureLevel::SM5) && (!IsMetalPlatform(Parameters.Platform) || RHIGetShaderLanguageVersion(Parameters.Platform) >= 2) && !IsMetalMRTPlatform(Parameters.Platform);
 		}
 
 
@@ -876,7 +876,7 @@ namespace GPUFFT
 		static bool ShouldCompilePermutation(const FGlobalShaderPermutationParameters& Parameters)
 		{
 			// @todo MetalMRT: Metal MRT can't cope with the threadgroup storage requirements for these shaders right now
-			return IsFeatureLevelSupported(Parameters.Platform, ERHIFeatureLevel::SM5) && (!IsMetalPlatform(Parameters.Platform) || RHIGetShaderLanguageVersion(Parameters.Platform) >= 2) && (Parameters.Platform != SP_METAL_MRT);
+			return IsFeatureLevelSupported(Parameters.Platform, ERHIFeatureLevel::SM5) && (!IsMetalPlatform(Parameters.Platform) || RHIGetShaderLanguageVersion(Parameters.Platform) >= 2) && !IsMetalMRTPlatform(Parameters.Platform);
 		}
 
 
@@ -1049,7 +1049,7 @@ namespace GPUFFT
 		static bool ShouldCompilePermutation(const FGlobalShaderPermutationParameters& Parameters)
 		{
 			// @todo MetalMRT: Metal MRT can't cope with the threadgroup storage requirements for these shaders right now
-			return IsFeatureLevelSupported(Parameters.Platform, ERHIFeatureLevel::SM5) && (!IsMetalPlatform(Parameters.Platform) || RHIGetShaderLanguageVersion(Parameters.Platform) >= 2) && (Parameters.Platform != SP_METAL_MRT);
+			return IsFeatureLevelSupported(Parameters.Platform, ERHIFeatureLevel::SM5) && (!IsMetalPlatform(Parameters.Platform) || RHIGetShaderLanguageVersion(Parameters.Platform) >= 2) && !IsMetalMRTPlatform(Parameters.Platform);
 		}
 
 
diff --git a/Engine/Source/Runtime/Renderer/Private/MobileShadingRenderer.cpp b/Engine/Source/Runtime/Renderer/Private/MobileShadingRenderer.cpp
--- a/Engine/Source/Runtime/Renderer/Private/MobileShadingRenderer.cpp
+++ b/Engine/Source/Runtime/Renderer/Private/MobileShadingRenderer.cpp
@@ -376,7 +376,7 @@ void FMobileSceneRenderer::Render(FRHICommandListImmediate& RHICmdList)
 		GSupportsShaderFramebufferFetch &&
 		ViewFamily.EngineShowFlags.PostProcessing;
 	bool bOnChipSunMask = bOnChipPP && (View.bLightShaftUse || GetMobileDepthOfFieldScale(View) > 0.0) && !IsMobileDistortionActive(View);
-	bool bOnChipPreResolveMSAA = bOnChipPP && (ViewFamily.GetShaderPlatform() == SP_METAL) && (CVarMobileMSAA ? CVarMobileMSAA->GetValueOnAnyThread() > 1 : false);
+	bool bOnChipPreResolveMSAA = bOnChipPP && IsMetalMobilePlatform(ViewFamily.GetShaderPlatform()) && (CVarMobileMSAA ? CVarMobileMSAA->GetValueOnAnyThread() > 1 : false);
 
 	if (!bGammaSpace && (bOnChipSunMask || bOnChipPreResolveMSAA))
 	{
diff --git a/Engine/Source/Runtime/Renderer/Private/PostProcess/DiaphragmDOF.h b/Engine/Source/Runtime/Renderer/Private/PostProcess/DiaphragmDOF.h
--- a/Engine/Source/Runtime/Renderer/Private/PostProcess/DiaphragmDOF.h
+++ b/Engine/Source/Runtime/Renderer/Private/PostProcess/DiaphragmDOF.h
@@ -122,10 +122,7 @@ inline bool IsSupported(EShaderPlatform ShaderPlatform)
 		ShaderPlatform == SP_XBOXONE_D3D12 ||
 		ShaderPlatform == SP_PS4 ||
 		IsVulkanSM5Platform(ShaderPlatform) ||
-		ShaderPlatform == SP_METAL_SM5 ||
-		ShaderPlatform == SP_METAL_SM5_NOTESS ||
-		ShaderPlatform == SP_METAL_MRT ||
-		ShaderPlatform == SP_METAL_MRT_MAC;
+		IsMetalSM5Platform(ShaderPlatform);
 }
 
 
diff --git a/Engine/Source/Runtime/Renderer/Private/PostProcess/PostProcessFFTBloom.cpp b/Engine/Source/Runtime/Renderer/Private/PostProcess/PostProcessFFTBloom.cpp
--- a/Engine/Source/Runtime/Renderer/Private/PostProcess/PostProcessFFTBloom.cpp
+++ b/Engine/Source/Runtime/Renderer/Private/PostProcess/PostProcessFFTBloom.cpp
@@ -42,7 +42,7 @@ public:
 	static bool ShouldCompilePermutation(const FGlobalShaderPermutationParameters& Parameters)
 	{
 		// @todo MetalMRT: Metal MRT can't cope with the threadgroup storage requirements for these shaders right now
-		return IsFeatureLevelSupported(Parameters.Platform, ERHIFeatureLevel::SM5) && (!IsMetalPlatform(Parameters.Platform) || RHIGetShaderLanguageVersion(Parameters.Platform) >= 2) && (Parameters.Platform != SP_METAL_MRT);
+		return IsFeatureLevelSupported(Parameters.Platform, ERHIFeatureLevel::SM5) && (!IsMetalPlatform(Parameters.Platform) || RHIGetShaderLanguageVersion(Parameters.Platform) >= 2) && !IsMetalMRTPlatform(Parameters.Platform);
 	}
 
 	static void ModifyCompilationEnvironment(const FGlobalShaderPermutationParameters& Parameters, FShaderCompilerEnvironment& OutEnvironment)
@@ -139,7 +139,7 @@ public:
 	static bool ShouldCompilePermutation(const FGlobalShaderPermutationParameters& Parameters)
 	{
 		// @todo MetalMRT: Metal MRT can't cope with the threadgroup storage requirements for these shaders right now
-		return IsFeatureLevelSupported(Parameters.Platform, ERHIFeatureLevel::SM5) && (!IsMetalPlatform(Parameters.Platform) || RHIGetShaderLanguageVersion(Parameters.Platform) >= 2) && (Parameters.Platform != SP_METAL_MRT);
+		return IsFeatureLevelSupported(Parameters.Platform, ERHIFeatureLevel::SM5) && (!IsMetalPlatform(Parameters.Platform) || RHIGetShaderLanguageVersion(Parameters.Platform) >= 2) && !IsMetalMRTPlatform(Parameters.Platform);
 	}
 
 	static void ModifyCompilationEnvironment(const FGlobalShaderPermutationParameters& Parameters, FShaderCompilerEnvironment& OutEnvironment)
@@ -228,7 +228,7 @@ public:
 	static bool ShouldCompilePermutation(const FGlobalShaderPermutationParameters& Parameters)
 	{
 		// @todo MetalMRT: Metal MRT can't cope with the threadgroup storage requirements for these shaders right now
-		return IsFeatureLevelSupported(Parameters.Platform, ERHIFeatureLevel::SM5) && (!IsMetalPlatform(Parameters.Platform) || RHIGetShaderLanguageVersion(Parameters.Platform) >= 2) && (Parameters.Platform != SP_METAL_MRT);
+		return IsFeatureLevelSupported(Parameters.Platform, ERHIFeatureLevel::SM5) && (!IsMetalPlatform(Parameters.Platform) || RHIGetShaderLanguageVersion(Parameters.Platform) >= 2) && !IsMetalMRTPlatform(Parameters.Platform);
 	}
 
 	static void ModifyCompilationEnvironment(const FGlobalShaderPermutationParameters& Parameters, FShaderCompilerEnvironment& OutEnvironment)
@@ -326,7 +326,7 @@ public:
 	static bool ShouldCompilePermutation(const FGlobalShaderPermutationParameters& Parameters)
 	{
 		// @todo MetalMRT: Metal MRT can't cope with the threadgroup storage requirements for these shaders right now
-		return IsFeatureLevelSupported(Parameters.Platform, ERHIFeatureLevel::SM5) && (!IsMetalPlatform(Parameters.Platform) || RHIGetShaderLanguageVersion(Parameters.Platform) >= 2) && (Parameters.Platform != SP_METAL_MRT);
+		return IsFeatureLevelSupported(Parameters.Platform, ERHIFeatureLevel::SM5) && (!IsMetalPlatform(Parameters.Platform) || RHIGetShaderLanguageVersion(Parameters.Platform) >= 2) && !IsMetalMRTPlatform(Parameters.Platform);
 	}
 
 	static void ModifyCompilationEnvironment(const FGlobalShaderPermutationParameters& Parameters, FShaderCompilerEnvironment& OutEnvironment)
diff --git a/Engine/Source/Runtime/Renderer/Private/PostProcess/PostProcessMobile.cpp b/Engine/Source/Runtime/Renderer/Private/PostProcess/PostProcessMobile.cpp
--- a/Engine/Source/Runtime/Renderer/Private/PostProcess/PostProcessMobile.cpp
+++ b/Engine/Source/Runtime/Renderer/Private/PostProcess/PostProcessMobile.cpp
@@ -178,7 +178,7 @@ void FRCPassPostProcessBloomSetupES2::SetShader(const FRenderingCompositePassCon
 
 	static const auto CVarMobileMSAA = IConsoleManager::Get().FindTConsoleVariableDataInt(TEXT("r.MobileMSAA"));
 	const EShaderPlatform ShaderPlatform = GShaderPlatformForFeatureLevel[Context.GetFeatureLevel()];
-	UseSunDof += (GSupportsShaderFramebufferFetch && (ShaderPlatform == SP_METAL || IsVulkanMobilePlatform(ShaderPlatform))) ? (CVarMobileMSAA ? (CVarMobileMSAA->GetValueOnRenderThread() > 1 ? 4 : 0) : 0) : 0;
+	UseSunDof += (GSupportsShaderFramebufferFetch && (IsMetalMobilePlatform(ShaderPlatform) || IsVulkanMobilePlatform(ShaderPlatform))) ? (CVarMobileMSAA ? (CVarMobileMSAA->GetValueOnRenderThread() > 1 ? 4 : 0) : 0) : 0;
 
 	switch(UseSunDof)
 	{
@@ -997,7 +997,7 @@ void FRCPassPostProcessSunMaskES2::SetShader(const FRenderingCompositePassContex
 	static const auto CVarMobileMSAA = IConsoleManager::Get().FindTConsoleVariableDataInt(TEXT("r.MobileMSAA"));
 	const EShaderPlatform ShaderPlatform = Context.GetShaderPlatform();
 
-	if ((GSupportsShaderFramebufferFetch && (ShaderPlatform == SP_METAL || IsVulkanMobilePlatform(ShaderPlatform))) && (CVarMobileMSAA ? CVarMobileMSAA->GetValueOnAnyThread() > 1 : false))
+	if ((GSupportsShaderFramebufferFetch && (IsMetalMobilePlatform(ShaderPlatform) || IsVulkanMobilePlatform(ShaderPlatform))) && (CVarMobileMSAA ? CVarMobileMSAA->GetValueOnAnyThread() > 1 : false))
 	{
 		UseFetchSunDof = 8;
 	}
diff --git a/Engine/Source/Runtime/Renderer/Private/PostProcess/PostProcessTonemap.cpp b/Engine/Source/Runtime/Renderer/Private/PostProcess/PostProcessTonemap.cpp
--- a/Engine/Source/Runtime/Renderer/Private/PostProcess/PostProcessTonemap.cpp
+++ b/Engine/Source/Runtime/Renderer/Private/PostProcess/PostProcessTonemap.cpp
@@ -1329,7 +1329,7 @@ class FPostProcessTonemapPS_ES2 : public FGlobalShader
 
 		static const auto CVarMobileMSAA = IConsoleManager::Get().FindTConsoleVariableDataInt(TEXT("r.MobileMSAA"));
 		const EShaderPlatform ShaderPlatform = GShaderPlatformForFeatureLevel[View.GetFeatureLevel()];
-		if ((GSupportsShaderFramebufferFetch && (ShaderPlatform == SP_METAL || IsVulkanMobilePlatform(ShaderPlatform))) && (CVarMobileMSAA ? CVarMobileMSAA->GetValueOnAnyThread() > 1 : false))
+		if ((GSupportsShaderFramebufferFetch && (IsMetalMobilePlatform(ShaderPlatform) || IsVulkanMobilePlatform(ShaderPlatform))) && (CVarMobileMSAA ? CVarMobileMSAA->GetValueOnAnyThread() > 1 : false))
 		{
 			MobilePermutationVector.Set<FTonemapperMsaaDim>(true);
 		}
diff --git a/Engine/Source/Runtime/Renderer/Private/PostProcess/PostProcessWeightedSampleSum.cpp b/Engine/Source/Runtime/Renderer/Private/PostProcess/PostProcessWeightedSampleSum.cpp
--- a/Engine/Source/Runtime/Renderer/Private/PostProcess/PostProcessWeightedSampleSum.cpp
+++ b/Engine/Source/Runtime/Renderer/Private/PostProcess/PostProcessWeightedSampleSum.cpp
@@ -71,11 +71,11 @@ public:
 
 	static bool ShouldCompilePermutation(const FGlobalShaderPermutationParameters& Parameters)
 	{
-		if( IsFeatureLevelSupported(Parameters.Platform, ERHIFeatureLevel::SM5) && (Parameters.Platform != SP_METAL_MRT && Parameters.Platform != SP_METAL_MRT_MAC) )
+		if( IsFeatureLevelSupported(Parameters.Platform, ERHIFeatureLevel::SM5) && !IsMetalMRTPlatform(Parameters.Platform) )
 		{
 			return true;
 		}
-		else if (Parameters.Platform == SP_METAL_MRT || Parameters.Platform == SP_METAL_MRT_MAC)
+		else if (IsMetalMRTPlatform(Parameters.Platform))
 		{
 			return CompileTimeNumSamples <= MAX_FILTER_COMPILE_TIME_SAMPLES_IOS;
 		}
@@ -229,11 +229,11 @@ public:
 
 	static bool ShouldCompilePermutation(const FGlobalShaderPermutationParameters& Parameters)
 	{
-		if( IsFeatureLevelSupported(Parameters.Platform, ERHIFeatureLevel::SM5) && (Parameters.Platform != SP_METAL_MRT && Parameters.Platform != SP_METAL_MRT_MAC) )
+		if( IsFeatureLevelSupported(Parameters.Platform, ERHIFeatureLevel::SM5) && !IsMetalMRTPlatform(Parameters.Platform) )
 		{
 			return true;
 		}
-		else if (Parameters.Platform == SP_METAL_MRT || Parameters.Platform == SP_METAL_MRT_MAC)
+		else if (IsMetalMRTPlatform(Parameters.Platform))
 		{
 			return NumSamples <= MAX_FILTER_COMPILE_TIME_SAMPLES_IOS;
 		}
@@ -1307,7 +1307,7 @@ uint32 FRCPassPostProcessWeightedSampleSum::GetMaxNumSamples(ERHIFeatureLevel::T
 
 	uint32 MaxNumSamples = MAX_FILTER_COMPILE_TIME_SAMPLES;
 
-	if (InPlatform == SP_METAL_MRT || InPlatform == SP_METAL_MRT_MAC)
+	if (IsMetalMRTPlatform(InPlatform))
 	{
 		MaxNumSamples = MAX_FILTER_COMPILE_TIME_SAMPLES_IOS;
 	}
diff --git a/Engine/Source/Runtime/Renderer/Private/PostProcess/PostProcessing.cpp b/Engine/Source/Runtime/Renderer/Private/PostProcess/PostProcessing.cpp
--- a/Engine/Source/Runtime/Renderer/Private/PostProcess/PostProcessing.cpp
+++ b/Engine/Source/Runtime/Renderer/Private/PostProcess/PostProcessing.cpp
@@ -2336,7 +2336,7 @@ void FPostProcessing::ProcessES2(FRHICommandListImmediate& RHICmdList, const FVi
 		// incorrect UTexture::SRGB state. (UTexture::SRGB != HW texture state)
 		bool bSRGBAwareTarget = View.Family->RenderTarget->GetDisplayGamma() == 1.0f
 			&& View.bIsSceneCapture
-			&& View.GetShaderPlatform() == EShaderPlatform::SP_METAL;
+			&& IsMetalMobilePlatform(View.GetShaderPlatform());
 
 		// add the passes we want to add to the graph (commenting a line means the pass is not inserted into the graph) ---------
 		if( View.Family->EngineShowFlags.PostProcessing && bAllowFullPostProcess)
diff --git a/Engine/Source/Runtime/ShaderCore/Private/ShaderCache.cpp b/Engine/Source/Runtime/ShaderCore/Private/ShaderCache.cpp
--- a/Engine/Source/Runtime/ShaderCore/Private/ShaderCache.cpp
+++ b/Engine/Source/Runtime/ShaderCore/Private/ShaderCache.cpp
@@ -651,6 +651,8 @@ static bool ShaderPlatformCanPrebindBoundShaderState(EShaderPlatform Platform)
 		case SP_PCD3D_ES2:
 		case SP_METAL:
 		case SP_METAL_MRT:
+		case SP_METAL_TVOS:
+		case SP_METAL_MRT_TVOS:
 		case SP_METAL_MRT_MAC:
 		case SP_METAL_SM5:
 		case SP_METAL_SM5_NOTESS:
