Duplicate 4504306: Fix handling of float3 Buffer/RWBuffer in Metal which was accidentally broken.

diff --git a/Engine/Source/Developer/Apple/MetalShaderFormat/Private/MetalBackend.cpp b/Engine/Source/Developer/Apple/MetalShaderFormat/Private/MetalBackend.cpp
--- a/Engine/Source/Developer/Apple/MetalShaderFormat/Private/MetalBackend.cpp
+++ b/Engine/Source/Developer/Apple/MetalShaderFormat/Private/MetalBackend.cpp
@@ -1376,7 +1376,7 @@ protected:
 				{
 					bool bIsStructuredBuffer = (inst->type->inner_type->is_record() || !strncmp(inst->type->name, "RWStructuredBuffer<", 19) || !strncmp(inst->type->name, "StructuredBuffer<", 17));
 					bool bIsByteAddressBuffer = (!strncmp(inst->type->name, "RWByteAddressBuffer", 19) || !strncmp(inst->type->name, "ByteAddressBuffer", 17));
-                	if (Buffers.AtomicVariables.find(inst) != Buffers.AtomicVariables.end() || bIsStructuredBuffer || bIsByteAddressBuffer || inst->invariant || (inst->type->components() == 3) || Backend.Version <= 2)
+					if (Buffers.AtomicVariables.find(inst) != Buffers.AtomicVariables.end() || bIsStructuredBuffer || bIsByteAddressBuffer || inst->invariant || (inst->type->components() == 3) || (inst->type->inner_type && inst->type->inner_type->components() == 3) || Backend.Version <= 2)
 					{
 						bInsertSideTable |= true;
 					}
