We can't include the Xcode version when compiling shaders remotely because the SCW process doesn't always have the SSH configuration required, which means we must always supply the '-c' option to the 'metal' compiler.

diff --git a/Engine/Source/Developer/Apple/MetalShaderFormat/Private/MetalDerivedData.cpp b/Engine/Source/Developer/Apple/MetalShaderFormat/Private/MetalDerivedData.cpp
--- a/Engine/Source/Developer/Apple/MetalShaderFormat/Private/MetalDerivedData.cpp
+++ b/Engine/Source/Developer/Apple/MetalShaderFormat/Private/MetalDerivedData.cpp
@@ -215,11 +215,11 @@ bool FMetalShaderBytecodeCooker::Build(TArray<uint8>& OutData)
 		if (bUseSharedPCH)
         {
             CopyLocalFileToRemote(Job.InputPCHFile, RemoteInputPCHFile);
-			MetalParams = FString::Printf(TEXT("-include-pch %s %s %s %s %s -Wno-null-character -fbracket-depth=1024 %s %s %s %s -o %s"), *RemoteInputPCHFile, *Job.MinOSVersion, *Job.DebugInfo, *Job.MathMode, XcodeMajorVers >= 10 ? TEXT("-c") : TEXT(""), *Job.Standard, *Job.Defines, *IncludeArgs, *RemoteInputFile, *RemoteObjFile);
+			MetalParams = FString::Printf(TEXT("-include-pch %s %s %s %s %s -Wno-null-character -fbracket-depth=1024 %s %s %s %s -o %s"), *RemoteInputPCHFile, *Job.MinOSVersion, *Job.DebugInfo, *Job.MathMode, TEXT("-c"), *Job.Standard, *Job.Defines, *IncludeArgs, *RemoteInputFile, *RemoteObjFile);
         }
         else
         {
-            MetalParams = FString::Printf(TEXT("%s %s %s %s -Wno-null-character -fbracket-depth=1024 %s %s %s %s -o %s"), *Job.MinOSVersion, *Job.DebugInfo, *Job.MathMode, XcodeMajorVers >= 10 ? TEXT("-c") : TEXT(""), *Job.Standard, *Job.Defines, *IncludeArgs, *RemoteInputFile, *RemoteObjFile);
+            MetalParams = FString::Printf(TEXT("%s %s %s %s -Wno-null-character -fbracket-depth=1024 %s %s %s %s -o %s"), *Job.MinOSVersion, *Job.DebugInfo, *Job.MathMode, TEXT("-c"), *Job.Standard, *Job.Defines, *IncludeArgs, *RemoteInputFile, *RemoteObjFile);
         }
 	}
 
diff --git a/Engine/Source/Developer/Apple/MetalShaderFormat/Private/MetalShaderCompiler.cpp b/Engine/Source/Developer/Apple/MetalShaderFormat/Private/MetalShaderCompiler.cpp
--- a/Engine/Source/Developer/Apple/MetalShaderFormat/Private/MetalShaderCompiler.cpp
+++ b/Engine/Source/Developer/Apple/MetalShaderFormat/Private/MetalShaderCompiler.cpp
@@ -544,7 +544,8 @@ uint16 GetXcodeVersion(uint64& BuildVersion)
 	{
 		Version = 0; // No Xcode install is 0, so only text shaders will work
 		FString XcodePath = GetXcodePath();
-		if (XcodePath.Len() > 0)
+		// Because of where and when this is called you can't invoke it on Win->Mac builds
+		if (XcodePath.Len() > 0 && PLATFORM_MAC)
 		{
 			FString Path = FString::Printf(TEXT("%s/usr/bin/xcodebuild"), *XcodePath);
 			FString Result;
