Remove bHasRegularUniformBuffers from the MetalShaderFormat because it is no longer required - Metal always uses real uniform buffers.

diff --git a/Engine/Source/Developer/Apple/MetalShaderFormat/Private/MetalShaderCompiler.cpp b/Engine/Source/Developer/Apple/MetalShaderFormat/Private/MetalShaderCompiler.cpp
--- a/Engine/Source/Developer/Apple/MetalShaderFormat/Private/MetalShaderCompiler.cpp
+++ b/Engine/Source/Developer/Apple/MetalShaderFormat/Private/MetalShaderCompiler.cpp
@@ -1117,8 +1117,6 @@ void BuildMetalShaderOutput(
         }
 	}
 
-	bool bHasRegularUniformBuffers = false;
-
 	// Then 'normal' uniform buffers.
 	for (auto& UniformBlock : CCHeader.UniformBlocks)
 	{
@@ -1129,7 +1127,6 @@ void BuildMetalShaderOutput(
 		}
 		UsedUniformBufferSlots[UBIndex] = true;
 		ParameterMap.AddParameterAllocation(*UniformBlock.Name, UBIndex, 0, 0);
-		bHasRegularUniformBuffers = true;
 	}
 
 	// Packed global uniforms
@@ -1150,75 +1147,21 @@ void BuildMetalShaderOutput(
 
 	// Packed Uniform Buffers
 	TMap<int, TMap<ANSICHAR, uint16> > PackedUniformBuffersSize;
-	if ((CCFlags & HLSLCC_PackUniformsIntoUniformBufferWithNames) == HLSLCC_PackUniformsIntoUniformBufferWithNames)
-	{
-		for (auto& PackedUB : CCHeader.PackedUBs)
-		{
-			for (auto& Member : PackedUB.Members)
-			{
-				ParameterMap.AddParameterAllocation(
-													*Member.Name,
-													EArrayType_FloatHighp,
-													Member.Offset * BytesPerComponent,
-													Member.Count * BytesPerComponent
-													);
-				
-				uint16& Size = PackedUniformBuffersSize.FindOrAdd(PackedUB.Attribute.Index).FindOrAdd(EArrayType_FloatHighp);
-				Size = FMath::Max<uint16>(BytesPerComponent * (Member.Offset + Member.Count), Size);
-			}
-		}
-	}
-	else
+	for (auto& PackedUB : CCHeader.PackedUBs)
 	{
-		for (auto& PackedUB : CCHeader.PackedUBs)
+		for (auto& Member : PackedUB.Members)
 		{
-			check(PackedUB.Attribute.Index == Header.Bindings.NumUniformBuffers);
-			UsedUniformBufferSlots[PackedUB.Attribute.Index] = true;
-			ParameterMap.AddParameterAllocation(*PackedUB.Attribute.Name, Header.Bindings.NumUniformBuffers++, 0, 0);
+			ParameterMap.AddParameterAllocation(
+												*Member.Name,
+												EArrayType_FloatHighp,
+												Member.Offset * BytesPerComponent,
+												Member.Count * BytesPerComponent
+												);
 			
-			// Nothing else...
-			//for (auto& Member : PackedUB.Members)
-			//{
-			//}
-		}
-		
-		// Packed Uniform Buffers copy lists & setup sizes for each UB/Precision entry
-		for (auto& PackedUBCopy : CCHeader.PackedUBCopies)
-		{
-			CrossCompiler::FUniformBufferCopyInfo CopyInfo;
-			CopyInfo.SourceUBIndex = PackedUBCopy.SourceUB;
-			CopyInfo.SourceOffsetInFloats = PackedUBCopy.SourceOffset;
-			CopyInfo.DestUBIndex = PackedUBCopy.DestUB;
-			CopyInfo.DestUBTypeName = PackedUBCopy.DestPackedType;
-			CopyInfo.DestUBTypeIndex = CrossCompiler::PackedTypeNameToTypeIndex(CopyInfo.DestUBTypeName);
-			CopyInfo.DestOffsetInFloats = PackedUBCopy.DestOffset;
-			CopyInfo.SizeInFloats = PackedUBCopy.Count;
-
-			Header.UniformBuffersCopyInfo.Add(CopyInfo);
-
-			auto& UniformBufferSize = PackedUniformBuffersSize.FindOrAdd(CopyInfo.DestUBIndex);
-			uint16& Size = UniformBufferSize.FindOrAdd(CopyInfo.DestUBTypeName);
-			Size = FMath::Max<uint16>(BytesPerComponent * (CopyInfo.DestOffsetInFloats + CopyInfo.SizeInFloats), Size);
-		}
-
-		for (auto& PackedUBCopy : CCHeader.PackedUBGlobalCopies)
-		{
-			CrossCompiler::FUniformBufferCopyInfo CopyInfo;
-			CopyInfo.SourceUBIndex = PackedUBCopy.SourceUB;
-			CopyInfo.SourceOffsetInFloats = PackedUBCopy.SourceOffset;
-			CopyInfo.DestUBIndex = PackedUBCopy.DestUB;
-			CopyInfo.DestUBTypeName = PackedUBCopy.DestPackedType;
-			CopyInfo.DestUBTypeIndex = CrossCompiler::PackedTypeNameToTypeIndex(CopyInfo.DestUBTypeName);
-			CopyInfo.DestOffsetInFloats = PackedUBCopy.DestOffset;
-			CopyInfo.SizeInFloats = PackedUBCopy.Count;
-
-			Header.UniformBuffersCopyInfo.Add(CopyInfo);
-
-			uint16& Size = PackedGlobalArraySize.FindOrAdd(CopyInfo.DestUBTypeName);
-			Size = FMath::Max<uint16>(BytesPerComponent * (CopyInfo.DestOffsetInFloats + CopyInfo.SizeInFloats), Size);
+			uint16& Size = PackedUniformBuffersSize.FindOrAdd(PackedUB.Attribute.Index).FindOrAdd(EArrayType_FloatHighp);
+			Size = FMath::Max<uint16>(BytesPerComponent * (Member.Offset + Member.Count), Size);
 		}
 	}
-	Header.Bindings.bHasRegularUniformBuffers = bHasRegularUniformBuffers;
 
 	// Setup Packed Array info
 	Header.Bindings.PackedGlobalArrays.Reserve(PackedGlobalArraySize.Num());
@@ -1237,48 +1180,22 @@ void BuildMetalShaderOutput(
 	// Setup Packed Uniform Buffers info
 	Header.Bindings.PackedUniformBuffers.Reserve(PackedUniformBuffersSize.Num());
 	
-	if ((CCFlags & HLSLCC_PackUniformsIntoUniformBufferWithNames) == HLSLCC_PackUniformsIntoUniformBufferWithNames)
-	{
-		// In this mode there should only be 0 or 1 packed UB that contains all the aligned & named global uniform parameters
-		check(PackedUniformBuffersSize.Num() <= 1);
-		for (auto Iterator = PackedUniformBuffersSize.CreateIterator(); Iterator; ++Iterator)
-		{
-			int BufferIndex = Iterator.Key();
-			auto& ArraySizes = Iterator.Value();
-			for (auto IterSizes = ArraySizes.CreateIterator(); IterSizes; ++IterSizes)
-			{
-				ANSICHAR TypeName = IterSizes.Key();
-				uint16 Size = IterSizes.Value();
-				Size = (Size + 0xf) & (~0xf);
-				CrossCompiler::FPackedArrayInfo Info;
-				Info.Size = Size;
-				Info.TypeName = TypeName;
-				Info.TypeIndex = BufferIndex;
-				Header.Bindings.PackedGlobalArrays.Add(Info);
-			}
-		}
-	}
-	else
-	{
-		for (auto Iterator = PackedUniformBuffersSize.CreateIterator(); Iterator; ++Iterator)
-		{
-			int BufferIndex = Iterator.Key();
-			auto& ArraySizes = Iterator.Value();
-			TArray<CrossCompiler::FPackedArrayInfo> InfoArray;
-			InfoArray.Reserve(ArraySizes.Num());
-			for (auto IterSizes = ArraySizes.CreateIterator(); IterSizes; ++IterSizes)
-			{
-				ANSICHAR TypeName = IterSizes.Key();
-				uint16 Size = IterSizes.Value();
-				Size = (Size + 0xf) & (~0xf);
-				CrossCompiler::FPackedArrayInfo Info;
-				Info.Size = Size;
-				Info.TypeName = TypeName;
-				Info.TypeIndex = CrossCompiler::PackedTypeNameToTypeIndex(TypeName);
-				InfoArray.Add(Info);
-			}
-			
-			Header.Bindings.PackedUniformBuffers.Add(InfoArray);
+	// In this mode there should only be 0 or 1 packed UB that contains all the aligned & named global uniform parameters
+	check(PackedUniformBuffersSize.Num() <= 1);
+	for (auto Iterator = PackedUniformBuffersSize.CreateIterator(); Iterator; ++Iterator)
+	{
+		int BufferIndex = Iterator.Key();
+		auto& ArraySizes = Iterator.Value();
+		for (auto IterSizes = ArraySizes.CreateIterator(); IterSizes; ++IterSizes)
+		{
+			ANSICHAR TypeName = IterSizes.Key();
+			uint16 Size = IterSizes.Value();
+			Size = (Size + 0xf) & (~0xf);
+			CrossCompiler::FPackedArrayInfo Info;
+			Info.Size = Size;
+			Info.TypeName = TypeName;
+			Info.TypeIndex = BufferIndex;
+			Header.Bindings.PackedGlobalArrays.Add(Info);
 		}
 	}
 
diff --git a/Engine/Source/Developer/Apple/MetalShaderFormat/Private/MetalShaderFormat.cpp b/Engine/Source/Developer/Apple/MetalShaderFormat/Private/MetalShaderFormat.cpp
--- a/Engine/Source/Developer/Apple/MetalShaderFormat/Private/MetalShaderFormat.cpp
+++ b/Engine/Source/Developer/Apple/MetalShaderFormat/Private/MetalShaderFormat.cpp
@@ -237,7 +237,7 @@ class FMetalShaderFormat : public IShaderFormat
 public:
 	enum
 	{
-		HEADER_VERSION = 58,
+		HEADER_VERSION = 59,
 	};
 	
 	struct FVersion
diff --git a/Engine/Source/Runtime/Apple/MetalRHI/Private/MetalCommands.cpp b/Engine/Source/Runtime/Apple/MetalRHI/Private/MetalCommands.cpp
--- a/Engine/Source/Runtime/Apple/MetalRHI/Private/MetalCommands.cpp
+++ b/Engine/Source/Runtime/Apple/MetalRHI/Private/MetalCommands.cpp
@@ -493,7 +493,7 @@ void FMetalRHICommandContext::RHISetShaderUniformBuffer(FVertexShaderRHIParamRef
 
 	auto& Bindings = VertexShader->Bindings;
 	check(BufferIndex < Bindings.NumUniformBuffers);
-	if (Bindings.bHasRegularUniformBuffers && (Bindings.ConstantBuffers & 1 << BufferIndex))
+	if (Bindings.ConstantBuffers & 1 << BufferIndex)
 	{
 		auto* UB = (FMetalUniformBuffer*)BufferRHI;
 		Context->GetCurrentState().SetShaderBuffer(SF_Vertex, UB->Buffer, UB->Data, 0, UB->GetSize(), BufferIndex);
@@ -509,7 +509,7 @@ void FMetalRHICommandContext::RHISetShaderUniformBuffer(FHullShaderRHIParamRef H
 
 	auto& Bindings = HullShader->Bindings;
 	check(BufferIndex < Bindings.NumUniformBuffers);
-	if (Bindings.bHasRegularUniformBuffers && (Bindings.ConstantBuffers & 1 << BufferIndex))
+	if (Bindings.ConstantBuffers & 1 << BufferIndex)
 	{
 		auto* UB = (FMetalUniformBuffer*)BufferRHI;
 		Context->GetCurrentState().SetShaderBuffer(SF_Hull, UB->Buffer, UB->Data, 0, UB->GetSize(), BufferIndex);
@@ -525,7 +525,7 @@ void FMetalRHICommandContext::RHISetShaderUniformBuffer(FDomainShaderRHIParamRef
 
 	auto& Bindings = DomainShader->Bindings;
 	check(BufferIndex < Bindings.NumUniformBuffers);
-	if (Bindings.bHasRegularUniformBuffers && (Bindings.ConstantBuffers & 1 << BufferIndex))
+	if (Bindings.ConstantBuffers & 1 << BufferIndex)
 	{
 		auto* UB = (FMetalUniformBuffer*)BufferRHI;
 		Context->GetCurrentState().SetShaderBuffer(SF_Domain, UB->Buffer, UB->Data, 0, UB->GetSize(), BufferIndex);
@@ -546,7 +546,7 @@ void FMetalRHICommandContext::RHISetShaderUniformBuffer(FPixelShaderRHIParamRef
 
 	auto& Bindings = PixelShader->Bindings;
 	check(BufferIndex < Bindings.NumUniformBuffers);
-	if (Bindings.bHasRegularUniformBuffers && (Bindings.ConstantBuffers & 1 << BufferIndex))
+	if (Bindings.ConstantBuffers & 1 << BufferIndex)
 	{
 		auto* UB = (FMetalUniformBuffer*)BufferRHI;
 		Context->GetCurrentState().SetShaderBuffer(SF_Pixel, UB->Buffer, UB->Data, 0, UB->GetSize(), BufferIndex);
@@ -562,7 +562,7 @@ void FMetalRHICommandContext::RHISetShaderUniformBuffer(FComputeShaderRHIParamRe
 
 	auto& Bindings = ComputeShader->Bindings;
 	check(BufferIndex < Bindings.NumUniformBuffers);
-	if (Bindings.bHasRegularUniformBuffers && (Bindings.ConstantBuffers & 1 << BufferIndex))
+	if (Bindings.ConstantBuffers & 1 << BufferIndex)
 	{
 		auto* UB = (FMetalUniformBuffer*)BufferRHI;
 		Context->GetCurrentState().SetShaderBuffer(SF_Compute, UB->Buffer, UB->Data, 0, UB->GetSize(), BufferIndex);
diff --git a/Engine/Source/Runtime/Apple/MetalRHI/Public/MetalShaderResources.h b/Engine/Source/Runtime/Apple/MetalRHI/Public/MetalShaderResources.h
--- a/Engine/Source/Runtime/Apple/MetalRHI/Public/MetalShaderResources.h
+++ b/Engine/Source/Runtime/Apple/MetalRHI/Public/MetalShaderResources.h
@@ -142,7 +142,6 @@ struct FMetalShaderBindings
 	uint8	NumSamplers;
 	uint8	NumUniformBuffers;
 	uint8	NumUAVs;
-	bool	bHasRegularUniformBuffers;
 	bool	bDiscards;
 
 	FMetalShaderBindings() :
@@ -154,7 +153,6 @@ struct FMetalShaderBindings
 		NumSamplers(0),
 		NumUniformBuffers(0),
 		NumUAVs(0),
-		bHasRegularUniformBuffers(false),
 		bDiscards(false)
 	{
 	}
@@ -174,7 +172,6 @@ inline FArchive& operator<<(FArchive& Ar, FMetalShaderBindings& Bindings)
 	Ar << Bindings.NumSamplers;
 	Ar << Bindings.NumUniformBuffers;
 	Ar << Bindings.NumUAVs;
-	Ar << Bindings.bHasRegularUniformBuffers;
 	Ar << Bindings.bDiscards;
 	return Ar;
 }
