Merging MetalShaderFormat from //UE4/Dev-Main@4395058 to Dev-Rendering (//UE4/Dev-Rendering)

diff --git a/Engine/Source/Developer/Apple/MetalShaderFormat/Private/MetalDerivedData.cpp b/Engine/Source/Developer/Apple/MetalShaderFormat/Private/MetalDerivedData.cpp
--- a/Engine/Source/Developer/Apple/MetalShaderFormat/Private/MetalDerivedData.cpp
+++ b/Engine/Source/Developer/Apple/MetalShaderFormat/Private/MetalDerivedData.cpp
@@ -28,6 +28,7 @@ extern FString GetMetalBinaryPath(uint32 ShaderPlatform);
 extern FString GetMetalToolsPath(uint32 ShaderPlatform);
 extern FString GetMetalLibraryPath(uint32 ShaderPlatform);
 extern FString GetMetalCompilerVersion(uint32 ShaderPlatform);
+extern uint16 GetXcodeVersion(uint64& BuildVersion);
 extern EShaderPlatform MetalShaderFormatToLegacyShaderPlatform(FName ShaderFormat);
 extern void BuildMetalShaderOutput(
 	FShaderCompilerOutput& ShaderOutput,
@@ -205,16 +206,20 @@ bool FMetalShaderBytecodeCooker::Build(TArray<uint8>& OutData)
 	{
 		CopyLocalFileToRemote(Job.InputFile, RemoteInputFile);
 
+		uint64 XcodeBuildVers = 0;
+		uint16 XcodeVers = GetXcodeVersion(XcodeBuildVers);
+		uint16 XcodeMajorVers = ((XcodeVers >> 8) & 0xff);
+		
 		// PCH 
 		bool bUseSharedPCH = Job.InputPCHFile.Len() && IFileManager::Get().FileExists(*Job.InputPCHFile);
 		if (bUseSharedPCH)
         {
             CopyLocalFileToRemote(Job.InputPCHFile, RemoteInputPCHFile);
-			MetalParams = FString::Printf(TEXT("-include-pch %s %s %s %s -c -Wno-null-character -fbracket-depth=1024 %s %s %s %s -o %s"), *RemoteInputPCHFile, *Job.MinOSVersion, *Job.DebugInfo, *Job.MathMode, *Job.Standard, *Job.Defines, *IncludeArgs, *RemoteInputFile, *RemoteObjFile);
+			MetalParams = FString::Printf(TEXT("-include-pch %s %s %s %s %s -Wno-null-character -fbracket-depth=1024 %s %s %s %s -o %s"), *RemoteInputPCHFile, *Job.MinOSVersion, *Job.DebugInfo, *Job.MathMode, XcodeMajorVers >= 10 ? TEXT("-c") : TEXT(""), *Job.Standard, *Job.Defines, *IncludeArgs, *RemoteInputFile, *RemoteObjFile);
         }
         else
         {
-            MetalParams = FString::Printf(TEXT("%s %s %s -c -Wno-null-character -fbracket-depth=1024 %s %s %s %s -o %s"), *Job.MinOSVersion, *Job.DebugInfo, *Job.MathMode, *Job.Standard, *Job.Defines, *IncludeArgs, *RemoteInputFile, *RemoteObjFile);
+            MetalParams = FString::Printf(TEXT("%s %s %s %s -Wno-null-character -fbracket-depth=1024 %s %s %s %s -o %s"), *Job.MinOSVersion, *Job.DebugInfo, *Job.MathMode, XcodeMajorVers >= 10 ? TEXT("-c") : TEXT(""), *Job.Standard, *Job.Defines, *IncludeArgs, *RemoteInputFile, *RemoteObjFile);
         }
 	}
 
diff --git a/Engine/Source/Developer/Apple/MetalShaderFormat/Private/MetalShaderCompiler.cpp b/Engine/Source/Developer/Apple/MetalShaderFormat/Private/MetalShaderCompiler.cpp
--- a/Engine/Source/Developer/Apple/MetalShaderFormat/Private/MetalShaderCompiler.cpp
+++ b/Engine/Source/Developer/Apple/MetalShaderFormat/Private/MetalShaderCompiler.cpp
@@ -10,6 +10,7 @@
 #include "Serialization/MemoryReader.h"
 #include "Misc/FileHelper.h"
 #include "Misc/Paths.h"
+#include "HAL/PlatformFilemanager.h"
 
 #if PLATFORM_WINDOWS
 #include "Windows/WindowsHWrapper.h"
@@ -2311,7 +2312,8 @@ bool StripShader_Metal(TArray<uint8>& Code, class FString const& DebugPath, bool
 			SourceCode.Append(SourceCodePtr, ShaderCode.GetActualShaderCodeSize() - CodeOffset);
 			
 			const ANSICHAR* ShaderSource = ShaderCode.FindOptionalData('c');
-			bool const bHasShaderSource = (ShaderSource && FCStringAnsi::Strlen(ShaderSource) > 0);
+			const size_t ShaderSourceLength = ShaderSource ? FCStringAnsi::Strlen(ShaderSource) : 0;
+			bool const bHasShaderSource = ShaderSourceLength > 0;
 			
 			const ANSICHAR* ShaderPath = ShaderCode.FindOptionalData('p');
 			bool const bHasShaderPath = (ShaderPath && FCStringAnsi::Strlen(ShaderPath) > 0);
@@ -2323,9 +2325,20 @@ bool StripShader_Metal(TArray<uint8>& Code, class FString const& DebugPath, bool
 				if (IFileManager::Get().MakeDirectory(*DebugFolderPath, true))
 				{
 					FString TempPath = FPaths::CreateTempFilename(*DebugFolderPath, TEXT("MetalShaderFile-"), TEXT(".metal"));
-					FFileHelper::SaveStringToFile(FString(ShaderSource), *TempPath);
-					IFileManager::Get().Move(*DebugFilePath, *TempPath, false, false, true, false);
-					IFileManager::Get().Delete(*TempPath);
+					IPlatformFile& PlatformFile = FPlatformFileManager::Get().GetPlatformFile();
+					IFileHandle* FileHandle = PlatformFile.OpenWrite(*TempPath);
+					if (FileHandle)
+					{
+						FileHandle->Write((const uint8 *)ShaderSource, ShaderSourceLength);
+						delete FileHandle;
+
+						IFileManager::Get().Move(*DebugFilePath, *TempPath, false, false, true, false);
+						IFileManager::Get().Delete(*TempPath);
+					}
+					else
+					{
+						UE_LOG(LogShaders, Error, TEXT("Shader stripping failed: shader %s (Len: %0.8x, CRC: %0.8x) failed to create file %s!"), *Header.ShaderName, Header.SourceLen, Header.SourceCRC, *TempPath);
+					}
 				}
 			}
 			
@@ -2474,7 +2487,7 @@ uint64 AppendShader_Metal(FName const& Format, FString const& WorkingDir, const
 						
 						InShaderCode = NewCode.GetReadAccess();
 						
-						UE_LOG(LogShaders, Display, TEXT("Archiving succeeded: shader %s (Len: %0.8x, CRC: %0.8x, SHA: %s)"), *Header.ShaderName, Header.SourceLen, Header.SourceCRC, *Hash.ToString());
+						UE_LOG(LogShaders, Verbose, TEXT("Archiving succeeded: shader %s (Len: %0.8x, CRC: %0.8x, SHA: %s)"), *Header.ShaderName, Header.SourceLen, Header.SourceCRC, *Hash.ToString());
 					}
 					else
 					{
@@ -2557,7 +2570,7 @@ bool FinalizeLibrary_Metal(FName const& Format, FString const& WorkingDir, FStri
 				uint32 CRC = (Shader & 0xffffffff);
 				
 				// Build source file name path
-				UE_LOG(LogShaders, Display, TEXT("[%d/%d] %s Main_%0.8x_%0.8x.o"), ++Index, Shaders.Num(), *Format.GetPlainNameString(), Len, CRC);
+				UE_LOG(LogShaders, Verbose, TEXT("[%d/%d] %s Main_%0.8x_%0.8x.o"), ++Index, Shaders.Num(), *Format.GetPlainNameString(), Len, CRC);
 				FString SourceFileNameParam = FString::Printf(TEXT("\"%s/Main_%0.8x_%0.8x.o\""), *WorkingDir, Len, CRC);
 				
 				// Remote builds copy file and swizzle Source File Name param
@@ -2620,25 +2633,49 @@ bool FinalizeLibrary_Metal(FName const& Format, FString const& WorkingDir, FStri
 				FString MetalLibPath = MetalToolsPath + TEXT("/metallib");
 				
 				FString RemoteLibPath = LocalPathToRemote(LibraryPath, RemoteDestination);
-				FString Params = FString::Printf(TEXT("-o=\"%s\" \"%s\""), *RemoteLibPath, *ArchivePath);
-					
+				FString OriginalRemoteLibPath = RemoteLibPath;
+				FString Params;
+
+				if (RemoteFileExists(RemoteLibPath))
+				{
+					UE_LOG(LogShaders, Warning, TEXT("Archiving warning: target metallib already exists and will be overwritten: %s"), *RemoteLibPath);
+				}
+				if (RemoveRemoteFile(RemoteLibPath) != 0)
+				{
+					UE_LOG(LogShaders, Warning, TEXT("Archiving warning: target metallib already exists and count not be overwritten: %s"), *RemoteLibPath);
+
+					// Output to a unique file
+					FGuid Guid = FGuid::NewGuid();
+					RemoteLibPath = OriginalRemoteLibPath + FString::Printf(TEXT(".%x%x%x%x"), Guid.A, Guid.B, Guid.C, Guid.D);
+				}
+			
+				Params = FString::Printf(TEXT("-o=\"%s\" \"%s\""), *RemoteLibPath, *ArchivePath);
+				ReturnCode = 0;
+				Results = TEXT("");
+				Errors = TEXT("");
 				ExecRemoteProcess( *MetalLibPath, *Params, &ReturnCode, &Results, &Errors );
-				
-				if(ReturnCode == 0)
-                {
-                    // There is problem going to location with spaces using remote copy (at least on Mac no combination of \ and/or "" works) - work around this issue @todo investigate this further
+	
+				// handle compile error
+				if (ReturnCode == 0)
+				{
+					// There is problem going to location with spaces using remote copy (at least on Mac no combination of \ and/or "" works) - work around this issue @todo investigate this further
                     FString LocalCopyLocation = FPaths::Combine(TEXT("/tmp"),FPaths::GetCleanFilename(LibraryPath));
 						
                     if(bBuildingRemotely && CopyRemoteFileToLocal(RemoteLibPath, LocalCopyLocation))
                     {
                         IFileManager::Get().Move(*LibraryPath, *LocalCopyLocation);
                     }
-                }
-				
-				// handle compile error
-				if (ReturnCode == 0 && IFileManager::Get().FileSize(*LibraryPath) > 0)
-				{
-					bOK = true;
+					else if (!bBuildingRemotely && RemoteLibPath != LibraryPath)
+					{
+						IFileManager::Get().Move(*RemoteLibPath, *LibraryPath);
+					}
+
+					bOK = (IFileManager::Get().FileSize(*LibraryPath) > 0);
+
+					if (!bOK)
+					{
+						UE_LOG(LogShaders, Error, TEXT("Archiving failed: failed to copy to local destination: %s"), *LibraryPath);
+					}
 				}
 				else
 				{
