Keep sampler names in Metal shaders - makes it easier to debug compilation.

diff --git a/Engine/Source/Developer/Apple/MetalShaderFormat/Private/MetalBackend.cpp b/Engine/Source/Developer/Apple/MetalShaderFormat/Private/MetalBackend.cpp
--- a/Engine/Source/Developer/Apple/MetalShaderFormat/Private/MetalBackend.cpp
+++ b/Engine/Source/Developer/Apple/MetalShaderFormat/Private/MetalBackend.cpp
@@ -1115,7 +1115,7 @@ protected:
 									{
 										ralloc_asprintf_append(
 											buffer,
-											"sampler s%d [[ sampler(%d) ]], ", SamplerStateIndex, SamplerStateIndex);
+											"sampler %s [[ sampler(%d) ]], ", SamplerState.c_str(), SamplerStateIndex);
 									}
 								}
 							}
@@ -1961,7 +1961,7 @@ protected:
 			bool bDummy;
 			int32 SamplerStateIndex = Buffers.GetUniqueSamplerStateIndex(tex->SamplerStateName, false, bDummy);
 			check(SamplerStateIndex != INDEX_NONE);
-			ralloc_asprintf_append(buffer, "s%d, ", SamplerStateIndex);
+			ralloc_asprintf_append(buffer, "%s, ", tex->SamplerStateName);
 			
 			bool bLocalCubeArrayHacks = false;
 			if (tex->sampler->type->sampler_array)
@@ -2197,7 +2197,7 @@ protected:
 			auto* Entry = ParseState->FindPackedSamplerEntry(Texture->name);
 			int32 SamplerStateIndex = Buffers.GetUniqueSamplerStateIndex(tex->SamplerStateName, false, bDummy);
 			check(SamplerStateIndex != INDEX_NONE);
-			ralloc_asprintf_append(buffer, "s%d, ", SamplerStateIndex);
+			ralloc_asprintf_append(buffer, "%s, ", tex->SamplerStateName);
 			// Coord
 			tex->coordinate->accept(this);
 
