Only zero-initialize numeric types in MetalRHI - can't zero-init references to structs/textures/samplers/buffers.

diff --git a/Engine/Source/Developer/Apple/MetalShaderFormat/Private/MetalBackend.cpp b/Engine/Source/Developer/Apple/MetalShaderFormat/Private/MetalBackend.cpp
--- a/Engine/Source/Developer/Apple/MetalShaderFormat/Private/MetalBackend.cpp
+++ b/Engine/Source/Developer/Apple/MetalShaderFormat/Private/MetalBackend.cpp
@@ -851,7 +851,7 @@ protected:
 
 	void print_zero_initialiser(const glsl_type * type)
 	{
-		if(type->base_type != GLSL_TYPE_STRUCT)
+		if(type->is_numeric() || (type->base_type == GLSL_TYPE_ARRAY))
 		{
 			if (type->base_type != GLSL_TYPE_ARRAY)
 			{
@@ -1343,7 +1343,7 @@ protected:
 		else if ((Backend.bZeroInitialise) && (var->type->base_type != GLSL_TYPE_STRUCT) && (var->mode == ir_var_auto || var->mode == ir_var_temporary || var->mode == ir_var_shared) && (Buffers.AtomicVariables.find(var) == Buffers.AtomicVariables.end()))
 		{
 			// @todo UE-34355 temporary workaround for 10.12 shader compiler error - really all arrays should be zero'd but only threadgroup shared initialisation works on the Beta drivers.
-			if (!is_struct_type(var->type) && (var->type->base_type != GLSL_TYPE_ARRAY || var->mode == ir_var_shared))
+			if (!is_struct_type(var->type) && (var->type->base_type != GLSL_TYPE_ARRAY || var->mode == ir_var_shared) && (var->type->is_numeric() || (var->type->base_type == GLSL_TYPE_ARRAY)))
 			{
 				ralloc_asprintf_append(buffer, " = ");
 				print_zero_initialiser(var->type);
