Duplicate 4628076 //UE4/Dev-Rendering to Dev-Main (//UE4/Dev-Main)

All macOS Metal shader formats must use Metal 2.0 or later now, I no longer support Metal 1.2 for desktop.

diff --git a/Engine/Source/Developer/Apple/MetalShaderFormat/Private/MetalShaderCompiler.cpp b/Engine/Source/Developer/Apple/MetalShaderFormat/Private/MetalShaderCompiler.cpp
--- a/Engine/Source/Developer/Apple/MetalShaderFormat/Private/MetalShaderCompiler.cpp
+++ b/Engine/Source/Developer/Apple/MetalShaderFormat/Private/MetalShaderCompiler.cpp
@@ -2001,43 +2001,43 @@ void CompileShader_Metal(const FShaderCompilerInput& _Input,FShaderCompilerOutpu
 	}
 	else if (Input.ShaderFormat == NAME_SF_METAL_MACES2)
 	{
-        UE_CLOG(VersionEnum == 0, LogShaders, Warning, TEXT("Metal shader version should be Metal v1.1 or higher for format %s!"), VersionEnum, *Input.ShaderFormat.ToString());
+        UE_CLOG(VersionEnum < 3, LogShaders, Warning, TEXT("Metal shader version must be Metal v2.0 or higher for format %s!"), VersionEnum, *Input.ShaderFormat.ToString());
 		AdditionalDefines.SetDefine(TEXT("METAL_ES2_PROFILE"), 1);
-		VersionEnum = VersionEnum > 0 ? VersionEnum : 1;
+		VersionEnum = VersionEnum >= 3 ? VersionEnum : 3;
 		MetalCompilerTarget = HCT_FeatureLevelES2;
 		Semantics = EMetalGPUSemanticsImmediateDesktop;
 	}
 	else if (Input.ShaderFormat == NAME_SF_METAL_MACES3_1)
 	{
-        UE_CLOG(VersionEnum == 0, LogShaders, Warning, TEXT("Metal shader version should be Metal v1.1 or higher for format %s!"), VersionEnum, *Input.ShaderFormat.ToString());
+        UE_CLOG(VersionEnum < 3, LogShaders, Warning, TEXT("Metal shader version must be Metal v2.0 or higher for format %s!"), VersionEnum, *Input.ShaderFormat.ToString());
 		AdditionalDefines.SetDefine(TEXT("METAL_PROFILE"), 1);
-		VersionEnum = VersionEnum > 0 ? VersionEnum : 1;
+		VersionEnum = VersionEnum >= 3 ? VersionEnum : 3;
 		MetalCompilerTarget = HCT_FeatureLevelES3_1;
 		Semantics = EMetalGPUSemanticsImmediateDesktop;
 	}
 	else if (Input.ShaderFormat == NAME_SF_METAL_SM5_NOTESS)
 	{
-        UE_CLOG(VersionEnum == 0, LogShaders, Warning, TEXT("Metal shader version should be Metal v1.2 or higher for format %s!"), VersionEnum, *Input.ShaderFormat.ToString());
+        UE_CLOG(VersionEnum < 3, LogShaders, Warning, TEXT("Metal shader version must be Metal v2.0 or higher for format %s!"), VersionEnum, *Input.ShaderFormat.ToString());
 		AdditionalDefines.SetDefine(TEXT("METAL_SM5_NOTESS_PROFILE"), 1);
 		AdditionalDefines.SetDefine(TEXT("USING_VERTEX_SHADER_LAYER"), 1);
-        VersionEnum = VersionEnum > 0 ? VersionEnum : 2;
+		VersionEnum = VersionEnum >= 3 ? VersionEnum : 3;
 		MetalCompilerTarget = HCT_FeatureLevelSM5;
 		Semantics = EMetalGPUSemanticsImmediateDesktop;
 	}
 	else if (Input.ShaderFormat == NAME_SF_METAL_SM5)
 	{
-        UE_CLOG(VersionEnum == 0, LogShaders, Warning, TEXT("Metal shader version should be Metal v1.2 or higher for format %s!"), VersionEnum, *Input.ShaderFormat.ToString());
+        UE_CLOG(VersionEnum < 3, LogShaders, Warning, TEXT("Metal shader version must be Metal v2.0 or higher for format %s!"), VersionEnum, *Input.ShaderFormat.ToString());
 		AdditionalDefines.SetDefine(TEXT("METAL_SM5_PROFILE"), 1);
 		AdditionalDefines.SetDefine(TEXT("USING_VERTEX_SHADER_LAYER"), 1);
-        VersionEnum = VersionEnum > 0 ? VersionEnum : 3;
+		VersionEnum = VersionEnum >= 3 ? VersionEnum : 3;
 		MetalCompilerTarget = HCT_FeatureLevelSM5;
 		Semantics = EMetalGPUSemanticsImmediateDesktop;
 	}
 	else if (Input.ShaderFormat == NAME_SF_METAL_MRT_MAC)
 	{
-        UE_CLOG(VersionEnum == 0, LogShaders, Warning, TEXT("Metal shader version should be Metal v1.2 or higher for format %s!"), VersionEnum, *Input.ShaderFormat.ToString());
+        UE_CLOG(VersionEnum < 3, LogShaders, Warning, TEXT("Metal shader version must be Metal v2.0 or higher for format %s!"), VersionEnum, *Input.ShaderFormat.ToString());
 		AdditionalDefines.SetDefine(TEXT("METAL_MRT_PROFILE"), 1);
-		VersionEnum = VersionEnum > 0 ? VersionEnum : 2;
+		VersionEnum = VersionEnum >= 3 ? VersionEnum : 3;
 		MetalCompilerTarget = HCT_FeatureLevelSM5;
 		Semantics = EMetalGPUSemanticsTBDRDesktop;
 	}
