Insert sampler-state ir_variables when compiling to Metal as we always support separate sampler objects - this makes it *much* easier to handle future changes.

diff --git a/Engine/Source/Developer/Apple/MetalShaderFormat/Private/MetalBackend.cpp b/Engine/Source/Developer/Apple/MetalShaderFormat/Private/MetalBackend.cpp
--- a/Engine/Source/Developer/Apple/MetalShaderFormat/Private/MetalBackend.cpp
+++ b/Engine/Source/Developer/Apple/MetalShaderFormat/Private/MetalBackend.cpp
@@ -1054,7 +1054,18 @@ protected:
 				{
 					auto* PtrType = var->type->is_array() ? var->type->element_type() : var->type;
 					check(!PtrType->is_array());
-					if (var->type->is_sampler())
+					if (var->type->base_type == GLSL_TYPE_SAMPLER_STATE)
+					{
+						bool bAdded = false;
+						int32 SamplerStateIndex = Buffers.GetUniqueSamplerStateIndex(var->name, true, bAdded);
+						if (bAdded)
+						{
+							ralloc_asprintf_append(
+												   buffer,
+												   "sampler %s [[ sampler(%d) ]]", var->name, SamplerStateIndex);
+						}
+					}
+					else if (var->type->is_sampler())
 					{
 						if (var->type->sampler_buffer)
 						{
@@ -1102,23 +1113,6 @@ protected:
 							// Regular textures
 							auto* Entry = ParseState->FindPackedSamplerEntry(var->name);
 							check(Entry);
-							//@todo-rco: SamplerStates
-							auto SamplerStateFound = ParseState->TextureToSamplerMap.find(Entry->Name.c_str());
-							if (SamplerStateFound != ParseState->TextureToSamplerMap.end())
-							{
-								auto& SamplerStates = SamplerStateFound->second;
-								for (auto& SamplerState : SamplerStates)
-								{
-									bool bAdded = false;
-									int32 SamplerStateIndex = Buffers.GetUniqueSamplerStateIndex(SamplerState, true, bAdded);
-									if (bAdded)
-									{
-										ralloc_asprintf_append(
-											buffer,
-											"sampler %s [[ sampler(%d) ]], ", SamplerState.c_str(), SamplerStateIndex);
-									}
-								}
-							}
 
 							print_type_pre(PtrType);
 							const char* InnerType = "float";
@@ -1957,11 +1951,21 @@ protected:
 			
 			auto* Texture = tex->sampler->variable_referenced();
 			check(Texture);
-			auto* Entry = ParseState->FindPackedSamplerEntry(Texture->name);
-			bool bDummy;
-			int32 SamplerStateIndex = Buffers.GetUniqueSamplerStateIndex(tex->SamplerStateName, false, bDummy);
-			check(SamplerStateIndex != INDEX_NONE);
-			ralloc_asprintf_append(buffer, "%s, ", tex->SamplerStateName);
+			
+			if (tex->SamplerState)
+			{
+				tex->SamplerState->accept(this);
+				ralloc_asprintf_append(buffer, ", ");
+			}
+			else
+			{
+				auto* Entry = ParseState->FindPackedSamplerEntry(Texture->name);
+				bool bDummy;
+				check(Entry);
+				int32 SamplerStateIndex = Buffers.GetUniqueSamplerStateIndex(tex->SamplerStateName, false, bDummy);
+				check(SamplerStateIndex != INDEX_NONE);
+				ralloc_asprintf_append(buffer, "%s, ", tex->SamplerStateName);
+			}
 			
 			bool bLocalCubeArrayHacks = false;
 			if (tex->sampler->type->sampler_array)
@@ -2193,11 +2197,22 @@ protected:
 			// Sampler
 			auto* Texture = tex->sampler->variable_referenced();
 			check(Texture);
-			bool bDummy;
-			auto* Entry = ParseState->FindPackedSamplerEntry(Texture->name);
-			int32 SamplerStateIndex = Buffers.GetUniqueSamplerStateIndex(tex->SamplerStateName, false, bDummy);
-			check(SamplerStateIndex != INDEX_NONE);
-			ralloc_asprintf_append(buffer, "%s, ", tex->SamplerStateName);
+			
+			if (tex->SamplerState)
+			{
+				tex->SamplerState->accept(this);
+				ralloc_asprintf_append(buffer, ", ");
+			}
+			else
+			{
+				bool bDummy;
+				auto* Entry = ParseState->FindPackedSamplerEntry(Texture->name);
+				check (Entry);
+				int32 SamplerStateIndex = Buffers.GetUniqueSamplerStateIndex(tex->SamplerStateName, false, bDummy);
+				check(SamplerStateIndex != INDEX_NONE);
+				ralloc_asprintf_append(buffer, "%s, ", tex->SamplerStateName);
+			}
+			
 			// Coord
 			tex->coordinate->accept(this);
 
@@ -4446,6 +4461,8 @@ char* FMetalCodeBackend::GenerateCode(exec_list* ir, _mesa_glsl_parse_state* sta
 		bool bConvertUniformsToFloats = (HlslCompileFlags & HLSLCC_FlattenUniformBuffers) != HLSLCC_FlattenUniformBuffers;
 		ConvertHalfToFloatUniformsAndSamples(ir, state, bConvertUniformsToFloats, true);
 		
+		InsertSamplerStates(ir, state);
+		
 		Validate(ir, state);
 	}
 
diff --git a/Engine/Source/Developer/Apple/MetalShaderFormat/Private/MetalBackend.h b/Engine/Source/Developer/Apple/MetalShaderFormat/Private/MetalBackend.h
--- a/Engine/Source/Developer/Apple/MetalShaderFormat/Private/MetalBackend.h
+++ b/Engine/Source/Developer/Apple/MetalShaderFormat/Private/MetalBackend.h
@@ -106,6 +106,9 @@ struct FMetalCodeBackend : public FCodeBackend
 	// Return false if there were restrictions that made compilation fail
 	virtual bool ApplyAndVerifyPlatformRestrictions(exec_list* Instructions, _mesa_glsl_parse_state* ParseState, EHlslShaderFrequency Frequency) override;
 
+	struct glsl_type const* create_iab_type(_mesa_glsl_parse_state* ParseState, struct glsl_type const* UBType, char const* n, FBuffers const& Buffers);
+	void build_iab_fields(_mesa_glsl_parse_state* ParseState, char const* n, struct glsl_type const* t, TArray<struct glsl_struct_field>& Fields, unsigned& FieldIndex, unsigned& BufferIndex, bool top, FBuffers const& Buffers);
+	void InsertArgumentBuffers(exec_list* ir, _mesa_glsl_parse_state* state, FBuffers& Buffers);
 	void PackInputsAndOutputs(exec_list* ir, _mesa_glsl_parse_state* state, EHlslShaderFrequency Frequency, exec_list& InputVars);
 	void MovePackedUniformsToMain(exec_list* ir, _mesa_glsl_parse_state* state, FBuffers& OutBuffers);
 	void FixIntrinsics(exec_list* ir, _mesa_glsl_parse_state* state);
@@ -114,7 +117,10 @@ struct FMetalCodeBackend : public FCodeBackend
 	void ConvertHalfToFloatUniformsAndSamples(exec_list* ir, _mesa_glsl_parse_state* State, bool bConvertUniforms, bool bConvertSamples);
 	void BreakPrecisionChangesVisitor(exec_list* ir, _mesa_glsl_parse_state* State);
 	void FixupMetalBaseOffsets(exec_list* ir, _mesa_glsl_parse_state* state, EHlslShaderFrequency Frequency);
+	void InsertSamplerStates(exec_list* ir, _mesa_glsl_parse_state* State);
 
+	TMap<ir_variable*, TSet<uint8>> IABVariableMask;
+	TMap<ir_variable*, ir_variable*> IABVariablesMap;
     TMap<ir_variable*, uint32> ImageRW;
     FMetalTessellationOutputs& TessAttribs;
 	TArray<uint8> TypedBufferFormats;
diff --git a/Engine/Source/Developer/Apple/MetalShaderFormat/Private/MetalUtils.cpp b/Engine/Source/Developer/Apple/MetalShaderFormat/Private/MetalUtils.cpp
--- a/Engine/Source/Developer/Apple/MetalShaderFormat/Private/MetalUtils.cpp
+++ b/Engine/Source/Developer/Apple/MetalShaderFormat/Private/MetalUtils.cpp
@@ -2921,6 +2921,60 @@ void FMetalCodeBackend::ConvertHalfToFloatUniformsAndSamples(exec_list* ir, _mes
 	}
 }
 
+struct FMetalInsertSamplerStates final : public ir_rvalue_visitor
+{
+	_mesa_glsl_parse_state* State;
+	TMap<FString, ir_variable*> SamplerStates;
+	
+	FMetalInsertSamplerStates(_mesa_glsl_parse_state* InState) : State(InState)
+	{}
+	virtual ~FMetalInsertSamplerStates() { }
+	
+	virtual ir_visitor_status visit_leave(ir_texture * tex_op)
+	{
+		switch (tex_op->op)
+		{
+			case ir_tex:
+			case ir_txl:
+			case ir_txb:
+			case ir_txd:
+			case ir_txg:
+				if (tex_op->SamplerStateName && strlen(tex_op->SamplerStateName) > 0 && !tex_op->SamplerState)
+				{
+					ir_variable* Sampler = SamplerStates.FindRef(tex_op->SamplerStateName);
+					if (!Sampler)
+					{
+						Sampler = new (State)ir_variable(glsl_type::sampler_state_type, ralloc_strdup(State, tex_op->SamplerStateName), ir_var_uniform);
+						SamplerStates.Add(tex_op->SamplerStateName, Sampler);
+					}
+					ir_dereference_variable* Deref = new (State)ir_dereference_variable(Sampler);
+					tex_op->SamplerState = Deref;
+				}
+				break;
+			default:
+				break;
+		}
+		return ir_rvalue_visitor::visit_leave(tex_op);
+	}
+	
+	virtual void handle_rvalue(ir_rvalue** RValuePtr) override
+	{
+	}
+};
+
+void FMetalCodeBackend::InsertSamplerStates(exec_list* ir, _mesa_glsl_parse_state* State)
+{
+	FMetalInsertSamplerStates Visitor(State);
+	Visitor.run(ir);
+	
+	ir_function_signature* MainSig = GetMainFunction(ir);
+	check(MainSig);
+	for (auto const& Pair : Visitor.SamplerStates)
+	{
+		MainSig->parameters.push_tail(Pair.Value);
+	}
+}
+
 struct FMetalBreakPrecisionChangesVisitor final : public ir_rvalue_visitor
 {
 	_mesa_glsl_parse_state* State;
