Remove unimplemented function declarations.

diff --git a/Engine/Source/Developer/Apple/MetalShaderFormat/Private/MetalBackend.h b/Engine/Source/Developer/Apple/MetalShaderFormat/Private/MetalBackend.h
--- a/Engine/Source/Developer/Apple/MetalShaderFormat/Private/MetalBackend.h
+++ b/Engine/Source/Developer/Apple/MetalShaderFormat/Private/MetalBackend.h
@@ -106,9 +106,6 @@ struct FMetalCodeBackend : public FCodeBackend
 	// Return false if there were restrictions that made compilation fail
 	virtual bool ApplyAndVerifyPlatformRestrictions(exec_list* Instructions, _mesa_glsl_parse_state* ParseState, EHlslShaderFrequency Frequency) override;
 
-	struct glsl_type const* create_iab_type(_mesa_glsl_parse_state* ParseState, struct glsl_type const* UBType, char const* n, FBuffers const& Buffers);
-	void build_iab_fields(_mesa_glsl_parse_state* ParseState, char const* n, struct glsl_type const* t, TArray<struct glsl_struct_field>& Fields, unsigned& FieldIndex, unsigned& BufferIndex, bool top, FBuffers const& Buffers);
-	void InsertArgumentBuffers(exec_list* ir, _mesa_glsl_parse_state* state, FBuffers& Buffers);
 	void PackInputsAndOutputs(exec_list* ir, _mesa_glsl_parse_state* state, EHlslShaderFrequency Frequency, exec_list& InputVars);
 	void MovePackedUniformsToMain(exec_list* ir, _mesa_glsl_parse_state* state, FBuffers& OutBuffers);
 	void FixIntrinsics(exec_list* ir, _mesa_glsl_parse_state* state);
