Use the custom-header parsing for the Metal side-table index and centralise the printing of Metal types to make it easier to comprehend.

diff --git a/Engine/Source/Developer/Apple/MetalShaderFormat/Private/MetalBackend.cpp b/Engine/Source/Developer/Apple/MetalShaderFormat/Private/MetalBackend.cpp
--- a/Engine/Source/Developer/Apple/MetalShaderFormat/Private/MetalBackend.cpp
+++ b/Engine/Source/Developer/Apple/MetalShaderFormat/Private/MetalBackend.cpp
@@ -699,60 +699,159 @@ protected:
 		{
 			ralloc_asprintf_append(buffer, "%s", t->name);
 		}
-		else if (t->base_type == GLSL_TYPE_SAMPLER && t->sampler_buffer) 
+		else if (t->base_type == GLSL_TYPE_IMAGE)
 		{
-			// Typed buffer read
-			check(t->inner_type);
-			print_base_type(t->inner_type);
+			if (t->sampler_buffer)
+			{
+				if (!strncmp(t->name, "RWBuffer<", 9))
+				{
+					ralloc_asprintf_append(buffer, "buffer_argument<");
+					print_type_pre(t->inner_type);
+					ralloc_asprintf_append(buffer, ", access::read_write>");
+				}
+				else
+				{
+					if (strncmp(t->HlslName, "RW", 2))
+					{
+						ralloc_asprintf_append(buffer, "const ");
+					}
+					print_type_pre(t->inner_type);
+					ralloc_asprintf_append(buffer, "*");
+				}
+			}
+			else
+			{
+				auto ImageToMetalType = [](const char* Src, char* Dest, SIZE_T DestLen)
+				{
+					auto* Found = strstr(Src, "image");
+					check(Found);
+					Src = Found + 5;	// strlen("image")
+					FCStringAnsi::Strcpy(Dest, DestLen, "texture");
+					Dest += 7;	// strlen("texture")
+					if (Src[0] >= '1' && Src[0] <= '3')
+					{
+						*Dest++ = *Src++;
+						*Dest++ = 'd';
+						*Dest = 0;
+						check(*Src == 'D');
+						Src++;
+					}
+					else if (strncmp(Src, "Cube", 4) == 0)
+					{
+						FCStringAnsi::Strcpy(Dest, DestLen, "cube");
+						Dest += 4;
+					}
+					else
+					{
+						check(0);
+					}
+					
+					if (strncmp(Src, "Array", 5) == 0)
+					{
+						FCStringAnsi::Strcpy(Dest, DestLen, "_array");
+					}
+				};
+				
+				check(t->inner_type->is_numeric());
+				char Temp[32];
+				ImageToMetalType(t->name, Temp, sizeof(Temp) - 1);
+				ralloc_asprintf_append(buffer, "%s<", Temp);
+				// UAVs require type per channel, not including # of channels
+				print_type_pre(t->inner_type->get_scalar_type());
+				if (t->HlslName && strncmp(t->HlslName, "RW", 2))
+				{
+					ralloc_asprintf_append(buffer, ", access::read>");
+				}
+				else
+				{
+					ralloc_asprintf_append(buffer, ", access::read_write>");
+				}
+			}
 		}
-		else if (t->base_type == GLSL_TYPE_IMAGE)
+		else if (t->base_type == GLSL_TYPE_SAMPLER_STATE)
 		{
-			// Do nothing...
+			ralloc_asprintf_append(buffer, "sampler");
 		}
 		else
 		{
 			if (t->base_type == GLSL_TYPE_SAMPLER)
 			{
-				bool bDone = false;
-				if (t->sampler_dimensionality == GLSL_SAMPLER_DIM_2D && t->sampler_array)
-				{
-					ralloc_asprintf_append(buffer, t->sampler_shadow ? "depth2d_array" : "texture2d_array");
-					bDone = true;
-				}
-				else if (t->sampler_dimensionality == GLSL_SAMPLER_DIM_CUBE && t->sampler_array)
+				glsl_sampler_dim TexType = t->sampler_buffer ? GLSL_SAMPLER_DIM_BUF : (glsl_sampler_dim)t->sampler_dimensionality;
+				
+				if (TexType < GLSL_SAMPLER_DIM_BUF)
 				{
-					if (Backend.bIsDesktop == EMetalGPUSemanticsImmediateDesktop)
+					if (t->sampler_shadow)
 					{
-						ralloc_asprintf_append(buffer, t->sampler_shadow ? "depthcube_array" : "texturecube_array");
+						ralloc_asprintf_append(buffer, "depth");
 					}
 					else
 					{
-						ralloc_asprintf_append(buffer, t->sampler_shadow ? "depth2d_array" : "texture2d_array");
+						ralloc_asprintf_append(buffer, "texture");
 					}
-					bDone = true;
 				}
-				else if (t->sampler_dimensionality == GLSL_SAMPLER_DIM_2D && t->sampler_ms)
+				
+				switch (TexType)
 				{
-					ralloc_asprintf_append(buffer, t->sampler_shadow ? "depth2d_ms" : "texture2d_ms");
-					bDone = true;
+					case GLSL_SAMPLER_DIM_1D:
+						ralloc_asprintf_append(buffer, "1d");
+						break;
+					case GLSL_SAMPLER_DIM_2D:
+						ralloc_asprintf_append(buffer, "2d");
+						break;
+					case GLSL_SAMPLER_DIM_3D:
+						ralloc_asprintf_append(buffer, "3d");
+						break;
+					case GLSL_SAMPLER_DIM_CUBE:
+						ralloc_asprintf_append(buffer, "cube");
+						break;
+					case GLSL_SAMPLER_DIM_BUF:
+						// Typed buffer read
+						check(t->inner_type);
+						ralloc_asprintf_append(buffer, "buffer_argument<");
+						print_base_type(t->inner_type);
+						ralloc_asprintf_append(buffer, ">");
+						break;
+					case GLSL_SAMPLER_DIM_RECT:
+					case GLSL_SAMPLER_DIM_EXTERNAL:
+					default:
+						check(false);
+						break;
 				}
-				else if (t->HlslName)
+				
+				if (TexType < GLSL_SAMPLER_DIM_BUF)
 				{
-					if (!strcmp(t->HlslName, "texture2d") && t->sampler_shadow)
+					if (t->sampler_ms)
 					{
-						ralloc_asprintf_append(buffer, "depth2d");
-						bDone = true;
+						ralloc_asprintf_append(buffer, "_ms");
 					}
-					else if (!strcmp(t->HlslName, "texturecube") && t->sampler_shadow)
+					if (t->sampler_array)
 					{
-						ralloc_asprintf_append(buffer, "depthcube");
-						bDone = true;
+						ralloc_asprintf_append(buffer, "_array");
 					}
-				}
-
-				if (!bDone)
-				{
-					ralloc_asprintf_append(buffer, "%s", t->HlslName ? t->HlslName : "UnsupportedSamplerType");
+					
+					const char* InnerType = "float";
+					if (t->inner_type)
+					{
+						//#todo-rco: Currently force to float...
+						if (!t->sampler_shadow)
+						{
+							switch (t->inner_type->base_type)
+							{
+								case GLSL_TYPE_HALF:
+									InnerType = "half";
+									break;
+								case GLSL_TYPE_INT:
+									InnerType = "int";
+									break;
+								case GLSL_TYPE_UINT:
+									InnerType = "uint";
+									break;
+								default:
+									break;
+							}
+						}
+					}
+					ralloc_asprintf_append(buffer, "<%s>", InnerType);
 				}
 			}
 			else
@@ -981,63 +1080,15 @@ protected:
 				}
 				else
 				{
-					auto ImageToMetalType = [](const char* Src, char* Dest, SIZE_T DestLen)
+					print_type_pre(PtrType);
+					if (var->mode != ir_var_temporary)
 					{
-						auto* Found = strstr(Src, "image");
-						check(Found);
-						Src = Found + 5;	// strlen("image")
-						FCStringAnsi::Strcpy(Dest, DestLen, "texture");
-						Dest += 7;	// strlen("texture")
-						if (Src[0] >= '1' && Src[0] <= '3')
-						{
-							*Dest++ = *Src++;
-							*Dest++ = 'd';
-							*Dest = 0;
-							check(*Src == 'D');
-							Src++;
-						}
-						else if (strncmp(Src, "Cube", 4) == 0)
-						{
-							FCStringAnsi::Strcpy(Dest, DestLen, "cube");
-							Dest += 4;
-						}
-						else
-						{
-							check(0);
-						}
-
-						if (strncmp(Src, "Array", 5) == 0)
-						{
-							FCStringAnsi::Strcpy(Dest, DestLen, "_array");
-						}
-					};
-
-					check(PtrType->inner_type->is_numeric());
-					char Temp[32];
-					ImageToMetalType(PtrType->name, Temp, sizeof(Temp) - 1);
-					ralloc_asprintf_append(buffer, "%s<", Temp);
-					// UAVs require type per channel, not including # of channels
-					print_type_pre(PtrType->inner_type->get_scalar_type());
-                    
-                    uint32 Access = Backend.ImageRW.FindChecked(var);
-                    switch((EMetalAccess)Access)
-                    {
-                        case EMetalAccessRead:
-                            ralloc_asprintf_append(buffer, ", access::read> %s", unique_name(var));
-                            break;
-                        case EMetalAccessWrite:
-                            ralloc_asprintf_append(buffer, ", access::write> %s", unique_name(var));
-                            break;
-                        case EMetalAccessReadWrite:
-                            ralloc_asprintf_append(buffer, ", access::read_write> %s", unique_name(var));
-                            break;
-                        default:
-                            check(false);
-                    }
-					ralloc_asprintf_append(
-						buffer,
-						" [[ texture(%d) ]]", BufferIndex
-						);
+						ralloc_asprintf_append(buffer, " %s [[ texture(%d) ]]", unique_name(var), BufferIndex);
+					}
+					else
+					{
+						ralloc_asprintf_append(buffer, " %s", unique_name(var));
+					}
 				}
 			}
 			else
@@ -1102,7 +1153,7 @@ protected:
 							else
 							{
                                 ralloc_asprintf_append(buffer, "typedBuffer%d_read(", PtrType->inner_type->components());
-                                print_type_pre(PtrType);
+                                print_type_pre(PtrType->inner_type);
                                 ralloc_asprintf_append(buffer, ", %s, %d)", unique_name(var), BufferIndex);
                                 Backend.TypedBufferFormats[BufferIndex] = GetBufferFormat(PtrType->inner_type);
                                 Backend.TypedBuffers |= (1 << BufferIndex);
@@ -1115,38 +1166,13 @@ protected:
 							check(Entry);
 
 							print_type_pre(PtrType);
-							const char* InnerType = "float";
-							if (PtrType->inner_type)
-							{
-								if (PtrType->base_type == GLSL_TYPE_SAMPLER && PtrType->sampler_shadow)
-								{
-									//#todo-rco: Currently force to float...
-								}
-								else
-								{
-									switch (PtrType->inner_type->base_type)
-									{
-									case GLSL_TYPE_HALF:
-										InnerType = "half";
-										break;
-									case GLSL_TYPE_INT:
-										InnerType = "int";
-										break;
-									case GLSL_TYPE_UINT:
-										InnerType = "uint";
-										break;
-									default:
-										break;
-									}
-								}
-							}
                             
                             int BufferIndex = Buffers.GetIndex(var);
                             check(BufferIndex >= 0);
                             
 							ralloc_asprintf_append(
 								buffer,
-								"<%s> %s", InnerType, unique_name(var));
+								" %s", unique_name(var));
 							print_type_post(PtrType);
 							ralloc_asprintf_append(
 								buffer,
@@ -1247,6 +1273,21 @@ protected:
 							var->semantic
 							);
 					}
+					else if(var->semantic && var->type->is_record())
+					{
+						ralloc_asprintf_append(buffer,"device ");
+						print_type_pre(var->type);
+						ralloc_asprintf_append(buffer,"& %s",unique_name(var));
+						print_type_post(var->type);
+						int BufferIndex = Buffers.GetIndex(var);
+						check(BufferIndex >= 0);
+						check(BufferIndex < 31);
+						ralloc_asprintf_append(
+							buffer,
+							" [[ buffer(%d) ]]",
+							BufferIndex
+							);
+					}
 					else
 					{
 						check(var->type->is_record());
@@ -3994,6 +4035,21 @@ protected:
 			ralloc_asprintf_append(buffer, "// @NumThreads: %d, %d, %d\n", this->NumThreadsX, this->NumThreadsY, this->NumThreadsZ);
 		}
 		
+		bool foundSideTable = false;
+		for (int i = 0; i < Buffers.Buffers.Num(); ++i)
+		{
+			if (Buffers.Buffers[i])
+			{
+				auto* Var = Buffers.Buffers[i]->as_variable();
+				if (!Var->type->is_sampler() && !Var->type->is_image() && Var->semantic && !strcmp(Var->semantic, "u") && Var->mode == ir_var_uniform && Var->name && !strcmp(Var->name, "BufferSizes"))
+				{
+					check(foundSideTable == false);
+					foundSideTable = true;
+					ralloc_asprintf_append(buffer, "// @SideTable: %s(%d)", Var->name, i);
+				}
+			}
+		}
+		
 		if (Backend.bIsTessellationVSHS || Frequency == tessellation_evaluation_shader)
 		{
 			check(tessellation.outputcontrolpoints != 0);
@@ -4138,21 +4194,6 @@ protected:
             }
 			ralloc_asprintf_append(buffer, "// @TessellationControlPointOutBuffer: %u\n", (uint32)patchControlIndex);
 		}
-
-		bool foundSideTable = false;
-		for (int i = 0; i < Buffers.Buffers.Num(); ++i)
-		{
-			if (Buffers.Buffers[i])
-			{
-				auto* Var = Buffers.Buffers[i]->as_variable();
-				if (!Var->type->is_sampler() && !Var->type->is_image() && Var->semantic && !strcmp(Var->semantic, "u") && Var->mode == ir_var_uniform && Var->name && !strcmp(Var->name, "BufferSizes"))
-				{
-					check(foundSideTable == false);
-					foundSideTable = true;
-					ralloc_asprintf_append(buffer, "// @SideTable: %s(%d)", Var->name, i);
-				}
-			}
-		}
 	}
 
 public:
diff --git a/Engine/Source/Developer/Apple/MetalShaderFormat/Private/MetalShaderCompiler.cpp b/Engine/Source/Developer/Apple/MetalShaderFormat/Private/MetalShaderCompiler.cpp
--- a/Engine/Source/Developer/Apple/MetalShaderFormat/Private/MetalShaderCompiler.cpp
+++ b/Engine/Source/Developer/Apple/MetalShaderFormat/Private/MetalShaderCompiler.cpp
@@ -678,6 +678,7 @@ struct FHlslccMetalHeader : public CrossCompiler::FHlslccHeader
 	uint32 TessellationControlPointIndexBuffer;
 	EMetalOutputWindingMode TessellationOutputWinding;
 	EMetalPartitionMode TessellationPartitioning;
+	int8 SideTable;
 	uint8 Version;
 	bool bUsingTessellation;
 };
@@ -699,6 +700,7 @@ FHlslccMetalHeader::FHlslccMetalHeader(uint8 const InVersion, bool const bInUsin
 	TessellationControlPointOutBuffer = UINT_MAX;
 	TessellationControlPointIndexBuffer = UINT_MAX;
 	
+	SideTable = -1;
 	Version = InVersion;
 	bUsingTessellation = bInUsingTessellation;
 }
@@ -726,8 +728,36 @@ static const int32 Str##PrefixLen = FCStringAnsi::Strlen(Str##Prefix)
 	DEF_PREFIX_STR(TessellationHSTFOutBuffer);
 	DEF_PREFIX_STR(TessellationControlPointOutBuffer);
 	DEF_PREFIX_STR(TessellationControlPointIndexBuffer);
+	DEF_PREFIX_STR(SideTable);
 #undef DEF_PREFIX_STR
 	
+	const ANSICHAR* SideTableString = FCStringAnsi::Strstr(ShaderSource, SideTablePrefix);
+	if (SideTableString)
+	{
+		ShaderSource = SideTableString;
+		ShaderSource += SideTablePrefixLen;
+		while (*ShaderSource && *ShaderSource != '\n')
+		{
+			if (*ShaderSource == '(')
+			{
+				ShaderSource++;
+				if (*ShaderSource && *ShaderSource != '\n')
+				{
+					SideTable = (int8)ParseNumber(ShaderSource);
+				}
+			}
+			else
+			{
+				ShaderSource++;
+			}
+		}
+		if (SideTable < 0)
+		{
+			UE_LOG(LogMetalShaderCompiler, Fatal, TEXT("Couldn't parse the SideTable buffer index for bounds checking"));
+			return false;
+		}
+	}
+	
 	// Early out for non-tessellation...
 	if (Version < 2 || !bUsingTessellation)
 	{
@@ -969,8 +999,6 @@ void BuildMetalShaderOutput(
 		UE_LOG(LogMetalShaderCompiler, Fatal, TEXT("Bad hlslcc header found"));
 	}
 	
-	const ANSICHAR* SideTableString = FCStringAnsi::Strstr(USFSource, "@SideTable: ");
-	
 	EShaderFrequency Frequency = (EShaderFrequency)ShaderOutput.Target.Frequency;
 	const bool bIsMobile = (ShaderInput.Target.Platform == SP_METAL || ShaderInput.Target.Platform == SP_METAL_MRT || ShaderInput.Target.Platform == SP_METAL_TVOS || ShaderInput.Target.Platform == SP_METAL_MRT_TVOS);
 	bool bNoFastMath = ShaderInput.Environment.CompilerFlags.Contains(CFLAG_NoFastMath);
@@ -1040,34 +1068,6 @@ void BuildMetalShaderOutput(
 		Header.Bindings.TypedBufferFormats = TypedBufferFormats;
 	}
 	
-	if (SideTableString)
-	{
-		int32 SideTableLoc = -1;
-		while (*SideTableString && *SideTableString != '\n')
-		{
-			if (*SideTableString == '(')
-			{
-				SideTableString++;
-				if (*SideTableString && *SideTableString != '\n')
-				{
-					SideTableLoc = (int32)ParseNumber(SideTableString);
-				}
-			}
-			else
-			{
-				SideTableString++;
-			}
-		}
-		if (SideTableLoc >= 0)
-		{
-			Header.SideTable = SideTableLoc;
-		}
-		else
-		{
-			UE_LOG(LogMetalShaderCompiler, Fatal, TEXT("Couldn't parse the SideTable buffer index for bounds checking"));
-		}
-	}
-	
 	FShaderParameterMap& ParameterMap = ShaderOutput.ParameterMap;
 
 	TBitArray<> UsedUniformBufferSlots;
@@ -1270,6 +1270,7 @@ void BuildMetalShaderOutput(
 	Header.TessellationOutputAttribs            = TessOutputAttribs;
 	Header.bTessFunctionConstants				= (FCStringAnsi::Strstr(USFSource, "indexBufferType [[ function_constant(32) ]]") != nullptr);
 	Header.bDeviceFunctionConstants				= (FCStringAnsi::Strstr(USFSource, "#define __METAL_DEVICE_CONSTANT_INDEX__ 1") != nullptr);
+	Header.SideTable 							= CCHeader.SideTable;
 	
 	// Build the SRT for this shader.
 	{
