Incomplete support for Metal 2.1 using Indirect Argument buffers to bind and access resources that were declared as part of a uniform buffer.

- Added a new Metal shader version (5) which compiles as Metal 2.1 with Indirect Argument Buffers for resource binding.
- Only enable the IAB feature for Metal shader version 5 and above and only on macOS for now.
- Added code to create argument encoders which are cached as creating them is pricey.
- Added code to create argument buffers from RHIUniformBuffers and their layouts.
- Ensured that vertex, index and uniform data is allocated such that the pointer won't be invalidated after being written into the IAB.
- Significant changes to the shader compiler and ue4_stdlib to emit IAB structures.

Todo:
- No support for atomic operations on elements within an IAB yet - need to modify ue4_stdlib and the MetalBackend.
- Incorrect compilation of StructuredBuffera and/or ByteAddressBuffer when within an IAB - need to emit a different type in the IAB structure and then only bind the buffer pointer at runtime.
- Indirect Argument Buffers are incompatible with the RHI thread because the creation of texture-views is deferred to the RHI thread but IAB creation is not yet.

diff --git a/Engine/Source/Developer/Apple/MetalShaderFormat/Private/MetalBackend.cpp b/Engine/Source/Developer/Apple/MetalShaderFormat/Private/MetalBackend.cpp
--- a/Engine/Source/Developer/Apple/MetalShaderFormat/Private/MetalBackend.cpp
+++ b/Engine/Source/Developer/Apple/MetalShaderFormat/Private/MetalBackend.cpp
@@ -1351,13 +1351,40 @@ protected:
 						ralloc_asprintf_append(buffer, "atomic_");
 					}
 
+					if (var->mode == ir_var_ref)
+					{
+						ir_instruction* ir = (ir_instruction*)var->next;
+						check(ir && ir->ir_type == ir_type_assignment);
+						ir_assignment* assign = (ir_assignment*)ir;
+						ir_variable* rhs = assign->rhs->variable_referenced();
+						if (rhs->mode == ir_var_uniform)
+						{
+							ralloc_asprintf_append(buffer, "constant ");
+						}
+						else if (rhs->mode == ir_var_in)
+						{
+							ralloc_asprintf_append(buffer, "device ");
+						}
+					}
 					print_type_pre(var->type);
+					if (var->mode == ir_var_ref)
+					{
+						ralloc_asprintf_append(buffer, "&");
+					}
 					ralloc_asprintf_append(buffer, " %s", unique_name(var));
 					print_type_post(var->type);
 					if(var->is_patch_constant)
 					{
 						ralloc_asprintf_append(buffer, "/*???, is_patch_constant*/");
 					}
+					if (var->mode == ir_var_ref)
+					{
+						ir_instruction* ir = (ir_instruction*)var->next;
+						check(ir && ir->ir_type == ir_type_assignment);
+						ralloc_asprintf_append(buffer, " = ");
+						ir_assignment* assign = (ir_assignment*)ir;
+						assign->rhs->accept(this);
+					}
 				}
 			}
 		}
@@ -2141,54 +2168,91 @@ protected:
 			if (tex->sampler->type->is_sampler() && tex->sampler->type->sampler_buffer)
 			{
 				auto* Texture = tex->sampler->variable_referenced();
-				int Index = Buffers.GetIndex(Texture);
-				check(Index >= 0);
-				
+				int Index = 0;
+				char const* BufferSizesName = "BufferSizes";
+				int FormatIndex = -1;
+				bool bSideTable = bInsertSideTable;
+				if (Texture->mode == ir_var_temporary)
+				{
+					// IAB sampling path
+					ir_variable* IABVariable = Backend.IABVariablesMap.FindChecked(Texture);
+					int FieldIndex = IABVariable->type->field_index(Texture->name);
+					for (int i = 0; i < FieldIndex; i++)
+					{
+						if (IABVariable->type->fields.structure[i].type->sampler_buffer)
+						{
+							Index++;
+						}
+					}
+					
+					BufferSizesName = ralloc_asprintf(ParseState, "%s.BufferSizes", IABVariable->name);
+					bSideTable = true;
+				}
+				else
+				{
+					// Function argument path
+					Index = Buffers.GetIndex(Texture);
+					FormatIndex = Index;
+				}
+				check(Index >= 0 && Index <= 30);
+					
 				ralloc_asprintf_append(buffer, "(");
 				
 				bool bIsStructuredBuffer = (Texture->type->inner_type->is_record() || !strncmp(Texture->type->name, "RWStructuredBuffer<", 19) || !strncmp(Texture->type->name, "StructuredBuffer<", 17));
 				bool bIsByteAddressBuffer = (!strncmp(Texture->type->name, "RWByteAddressBuffer", 19) || !strncmp(Texture->type->name, "ByteAddressBuffer", 17));
-                bool bIsInvariantType = Texture->invariant;
-                bool bIsAtomic = Buffers.AtomicVariables.find(Texture) != Buffers.AtomicVariables.end();
+				bool bIsInvariantType = Texture->invariant;
+				bool bIsAtomic = Buffers.AtomicVariables.find(Texture) != Buffers.AtomicVariables.end();
 				if (!bIsStructuredBuffer && !bIsByteAddressBuffer && !bIsInvariantType && !bIsAtomic)
 				{
-                    ralloc_asprintf_append(buffer, "buffer::load<");
-                    print_type_pre(Texture->type->inner_type);
-                    ralloc_asprintf_append(buffer, ", %d>(", Index);
+					ralloc_asprintf_append(buffer, "buffer::load<");
+					print_type_pre(Texture->type->inner_type);
+					ralloc_asprintf_append(buffer, ", %d>(", Index);
 					tex->sampler->accept(this);
 					ralloc_asprintf_append(buffer, ", ");
 					tex->coordinate->accept(this);
-					if (bInsertSideTable)
-						ralloc_asprintf_append(buffer, ", GMetalTypedBufferFormat%d, BufferSizes)", Index);
+					if (FormatIndex >= 0)
+					{
+						ralloc_asprintf_append(buffer, ", GMetalTypedBufferFormat%d", Index);
+					}
 					else
-						ralloc_asprintf_append(buffer, ", GMetalTypedBufferFormat%d)", Index);
+					{
+						ralloc_asprintf_append(buffer, ", GMetalUnknownFormat");
+					}
+					if (bSideTable)
+					{
+						ralloc_asprintf_append(buffer, ", %s)", BufferSizesName);
+					}
+					else
+					{
+						ralloc_asprintf_append(buffer, ")");
+					}
 				}
 				else if (Backend.bBoundsChecks)
 				{
 					check(Index <= 30);
 					
-                    if (!bIsAtomic && (!bIsStructuredBuffer || !Texture->type->inner_type->is_record()))
+					if (!bIsAtomic && (!bIsStructuredBuffer || !Texture->type->inner_type->is_record()))
 					{
-                        ralloc_asprintf_append(buffer, "buffer::load<");
-                        print_type_pre(Texture->type->inner_type);
-                        ralloc_asprintf_append(buffer, ", %d>(", Index);
-                        tex->sampler->accept(this);
-                        ralloc_asprintf_append(buffer, ", ");
-                        tex->coordinate->accept(this);
-						if (bInsertSideTable)
-	                        ralloc_asprintf_append(buffer, ", BufferSizes)");
+						ralloc_asprintf_append(buffer, "buffer::load<");
+						print_type_pre(Texture->type->inner_type);
+						ralloc_asprintf_append(buffer, ", %d>(", Index);
+						tex->sampler->accept(this);
+						ralloc_asprintf_append(buffer, ", ");
+						tex->coordinate->accept(this);
+						if (bSideTable)
+							ralloc_asprintf_append(buffer, ", %s)", BufferSizesName);
+					}
+					else
+					{
+						tex->sampler->accept(this);
+						ralloc_asprintf_append(buffer, "[");
+						ralloc_asprintf_append(buffer, "min(");
+						tex->coordinate->accept(this);
+						ralloc_asprintf_append(buffer, ",");
+						ralloc_asprintf_append(buffer, "((%s[%d] / sizeof(", BufferSizesName, Index);
+						print_type_pre(Texture->type->inner_type);
+						ralloc_asprintf_append(buffer, "))))]");
 					}
-                    else
-                    {
-                        tex->sampler->accept(this);
-                        ralloc_asprintf_append(buffer, "[");
-                        ralloc_asprintf_append(buffer, "min(");
-                        tex->coordinate->accept(this);
-                        ralloc_asprintf_append(buffer, ",");
-                        ralloc_asprintf_append(buffer, "((BufferSizes[%d] / sizeof(", Index);
-                        print_type_pre(Texture->type->inner_type);
-                        ralloc_asprintf_append(buffer, "))))]");
-                    }
 				}
 				else
 				{
@@ -2433,6 +2497,10 @@ protected:
 			{
 				hash_table_insert(used_uniform_blocks, (void*)var->semantic, var->semantic);
 			}
+			if (hash_table_find(used_uniform_blocks, var->name) == NULL)
+			{
+				hash_table_insert(used_uniform_blocks, (void*)var->name, var->name);
+			}
 		}
 	}
 
@@ -2514,7 +2582,32 @@ protected:
 				else
 				{
 					auto* Texture = deref->image->variable_referenced();
-					int Index = Buffers.GetIndex(Texture);
+					int Index = 0;
+					char const* BufferSizesName = "BufferSizes";
+					int FormatIndex = -1;
+					bool bSideTable = bInsertSideTable;
+					if (Texture->mode == ir_var_temporary)
+					{
+						// IAB sampling path
+						ir_variable* IABVariable = Backend.IABVariablesMap.FindChecked(Texture);
+						int FieldIndex = IABVariable->type->field_index(Texture->name);
+						for (int i = 0; i < FieldIndex; i++)
+						{
+							if (IABVariable->type->fields.structure[i].type->sampler_buffer)
+							{
+								Index++;
+							}
+						}
+						
+						BufferSizesName = ralloc_asprintf(ParseState, "%s.BufferSizes", IABVariable->name);
+						bSideTable = true;
+					}
+					else
+					{
+						// Function argument path
+						Index = Buffers.GetIndex(Texture);
+						FormatIndex = Index;
+					}
 					check(Index >= 0 && Index <= 30);
 					
 					ralloc_asprintf_append(buffer, "(");
@@ -2531,10 +2624,22 @@ protected:
 						deref->image->accept(this);
 						ralloc_asprintf_append(buffer, ", ");
 						deref->image_index->accept(this);
-						if (bInsertSideTable)
-							ralloc_asprintf_append(buffer, ", GMetalTypedBufferFormat%d, BufferSizes)", Index);
+						if (FormatIndex >= 0)
+						{
+							ralloc_asprintf_append(buffer, ", GMetalTypedBufferFormat%d", Index);
+						}
+						else
+						{
+							ralloc_asprintf_append(buffer, ", GMetalUnknownFormat");
+						}
+						if (bSideTable)
+						{
+							ralloc_asprintf_append(buffer, ", %s)", BufferSizesName);
+						}
 						else
-							ralloc_asprintf_append(buffer, ", GMetalTypedBufferFormat%d)", Index);
+						{
+							ralloc_asprintf_append(buffer, ")");
+						}
 					}
 					else if (Backend.bBoundsChecks)
 					{
@@ -2547,8 +2652,8 @@ protected:
                             deref->image->accept(this);
                             ralloc_asprintf_append(buffer, ", ");
                             deref->image_index->accept(this);
-							if (bInsertSideTable)
-	                            ralloc_asprintf_append(buffer, ", BufferSizes)");
+							if (bSideTable)
+								ralloc_asprintf_append(buffer, ", %s)", BufferSizesName);
                         }
                         else
                         {
@@ -2557,7 +2662,7 @@ protected:
                             ralloc_asprintf_append(buffer, "min(");
                             deref->image_index->accept(this);
                             ralloc_asprintf_append(buffer, ",");
-                            ralloc_asprintf_append(buffer, "((BufferSizes[%d] / sizeof(", Index);
+							ralloc_asprintf_append(buffer, "((%s[%d] / sizeof(", BufferSizesName, Index);
                             print_type_pre(Texture->type->inner_type);
                             ralloc_asprintf_append(buffer, "))))]");
                         }
@@ -2676,41 +2781,75 @@ protected:
 				else
 				{
 					auto* Texture = deref->image->variable_referenced();
+					int Index = 0;
+					char const* BufferSizesName = "BufferSizes";
+					int FormatIndex = -1;
+					bool bSideTable = bInsertSideTable;
+					if (Texture->mode == ir_var_temporary)
+					{
+						// IAB sampling path
+						ir_variable* IABVariable = Backend.IABVariablesMap.FindChecked(Texture);
+						int FieldIndex = IABVariable->type->field_index(Texture->name);
+						for (int i = 0; i < FieldIndex; i++)
+						{
+							if (IABVariable->type->fields.structure[i].type->sampler_buffer)
+							{
+								Index++;
+							}
+						}
+						
+						BufferSizesName = ralloc_asprintf(ParseState, "%s.BufferSizes", IABVariable->name);
+						bSideTable = true;
+					}
+					else
+					{
+						// Function argument path
+						Index = Buffers.GetIndex(Texture);
+						FormatIndex = Index;
+					}
+					check(Index >= 0 && Index <= 30);
+					
 					bool bIsStructuredBuffer = (Texture->type->inner_type->is_record() || !strncmp(Texture->type->name, "RWStructuredBuffer<", 19) || !strncmp(Texture->type->name, "StructuredBuffer<", 17));
 					bool bIsByteAddressBuffer = (!strncmp(Texture->type->name, "RWByteAddressBuffer", 19) || !strncmp(Texture->type->name, "ByteAddressBuffer", 17));
                     bool bIsInvariantType = Texture->invariant;
                     bool bIsAtomic = Buffers.AtomicVariables.find(Texture) != Buffers.AtomicVariables.end();
 					if (!bIsStructuredBuffer && !bIsByteAddressBuffer && !bIsInvariantType && !bIsAtomic)
 					{
-						int Index = Buffers.GetIndex(Texture);
-						check(Index >= 0 && Index <= 30);
-						
-                        ralloc_asprintf_append(buffer, "buffer::store<");
+						ralloc_asprintf_append(buffer, "buffer::store<");
                         print_type_pre(Texture->type->inner_type);
                         ralloc_asprintf_append(buffer, ", %d>(", Index);
 						deref->image->accept(this);
 						ralloc_asprintf_append(buffer, ", ");
 						deref->image_index->accept(this);
-						if (bInsertSideTable)
-							ralloc_asprintf_append(buffer, ", GMetalTypedBufferFormat%d, BufferSizes, ", Index);
+						if (FormatIndex >= 0)
+						{
+							ralloc_asprintf_append(buffer, ", GMetalTypedBufferFormat%d", Index);
+						}
 						else
-							ralloc_asprintf_append(buffer, ", GMetalTypedBufferFormat%d, ", Index);
+						{
+							ralloc_asprintf_append(buffer, ", GMetalUnknownFormat");
+						}
+						if (bSideTable)
+						{
+							ralloc_asprintf_append(buffer, ", %s, ", BufferSizesName);
+						}
+						else
+						{
+							ralloc_asprintf_append(buffer, ", ");
+						}
 						src->accept(this);
 						ralloc_asprintf_append(buffer, ")");
 					}
 					else if (Backend.bBoundsChecks)
 					{
-                        int Index = Buffers.GetIndex(Texture);
-                        check(Index >= 0 && Index <= 30);
-                        
                         ralloc_asprintf_append(buffer, "buffer::store<");
                         print_type_pre(Texture->type->inner_type);
                         ralloc_asprintf_append(buffer, ", %d>(", Index);
                         deref->image->accept(this);
                         ralloc_asprintf_append(buffer, ", ");
                         deref->image_index->accept(this);
-						if (bInsertSideTable)
-	                        ralloc_asprintf_append(buffer, ", BufferSizes, ");
+						if (bSideTable)
+							ralloc_asprintf_append(buffer, ", %s, ", BufferSizesName);
 						else
 	                        ralloc_asprintf_append(buffer, ", ");
                         src->accept(this);
@@ -2767,7 +2906,7 @@ protected:
 
 		// constant variables with initializers are statically assigned
 		ir_variable *var = assign->lhs->variable_referenced();
-		if (var->has_initializer && var->read_only && (var->constant_initializer || var->constant_value))
+		if ((var->has_initializer && var->read_only && (var->constant_initializer || var->constant_value)) || var->mode == ir_var_ref)
 		{
 			//This will leave a blank line with a semi-colon
 			return;
@@ -3504,7 +3643,34 @@ protected:
 						ralloc_asprintf_append(buffer, "alignas(4) ", "");
 					}
 					
+					
+					const glsl_type* t = s->fields.structure[j].type;
+					if (t->base_type == GLSL_TYPE_IMAGE)
+					{
+						if (t->sampler_buffer)
+						{
+							if (strncmp(t->name, "RWBuffer<", 9))
+							{
+								if (s->fields.structure[j].patchconstant)
+								{
+									ralloc_asprintf_append(buffer, "constant ");
+								}
+								else
+								{
+									ralloc_asprintf_append(buffer, "device ");
+								}
+							}
+						}
+					}
+					if (t->base_type == GLSL_TYPE_STRUCT && !strncmp(t->name, "CB_", 3))
+					{
+						ralloc_asprintf_append(buffer, "constant ");
+					}
 					print_type_pre(s->fields.structure[j].type);
+					if (t->base_type == GLSL_TYPE_STRUCT && !strncmp(t->name, "CB_", 3))
+					{
+						ralloc_asprintf_append(buffer, "&");
+					}
 					ralloc_asprintf_append(buffer, " %s", s->fields.structure[j].name);
 					if (!s->fields.structure[j].semantic || strncmp(s->fields.structure[j].semantic, "[[", 2))
 					{
@@ -3962,11 +4128,11 @@ protected:
 				if (Buffers.Buffers[i])
 				{
 					auto* Var = Buffers.Buffers[i]->as_variable();
-					if (!Var->semantic && !Var->type->is_sampler() && !Var->type->is_image() && state->CBPackedArraysMap.find(Var->name) == state->CBPackedArraysMap.end())
+					if ((!Var->semantic || !strncmp(Var->type->name, "IAB_", 3)) && !Var->type->is_sampler() && !Var->type->is_image() && state->CBPackedArraysMap.find(Var->name) == state->CBPackedArraysMap.end())
 					{
 						ralloc_asprintf_append(buffer, "%s%s(%d)",
 							bFirst ? "// @UniformBlocks: " : ",",
-							Var->name, i);
+							Var->semantic ? Var->semantic : Var->name, i);
 						bFirst = false;
 					}
 				}
@@ -4050,6 +4216,34 @@ protected:
 			}
 		}
 		
+		if (Backend.IABVariableMask.Num())
+		{
+			bool bComma = false;
+			ralloc_asprintf_append(buffer, "// @ArgumentBuffers: ");
+			for (auto const& Pair : Backend.IABVariableMask)
+			{
+				if (bComma)
+				{
+					ralloc_asprintf_append(buffer, ",");
+				}
+				int Index = Buffers.GetIndex(Pair.Key);
+				ralloc_asprintf_append(buffer, "%d[", Index);
+				bool bSetComma = false;
+				for (auto Mask : Pair.Value)
+				{
+					if (bSetComma)
+					{
+						ralloc_asprintf_append(buffer, ",");
+					}
+					ralloc_asprintf_append(buffer, "%u", (uint32)Mask);
+					bSetComma = true;
+				}
+				ralloc_asprintf_append(buffer, "]", Pair.Key->name, Index);
+				bComma = true;
+			}
+			ralloc_asprintf_append(buffer, "\n");
+		}
+		
 		if (Backend.bIsTessellationVSHS || Frequency == tessellation_evaluation_shader)
 		{
 			check(tessellation.outputcontrolpoints != 0);
@@ -4504,6 +4698,11 @@ char* FMetalCodeBackend::GenerateCode(exec_list* ir, _mesa_glsl_parse_state* sta
 		
 		InsertSamplerStates(ir, state);
 		
+		if (Version >= 5 && bIsDesktop == EMetalGPUSemanticsImmediateDesktop)
+		{
+			InsertArgumentBuffers(ir, state, Buffers);
+		}
+		
 		Validate(ir, state);
 	}
 
diff --git a/Engine/Source/Developer/Apple/MetalShaderFormat/Private/MetalShaderCompiler.cpp b/Engine/Source/Developer/Apple/MetalShaderFormat/Private/MetalShaderCompiler.cpp
--- a/Engine/Source/Developer/Apple/MetalShaderFormat/Private/MetalShaderCompiler.cpp
+++ b/Engine/Source/Developer/Apple/MetalShaderFormat/Private/MetalShaderCompiler.cpp
@@ -678,6 +678,7 @@ struct FHlslccMetalHeader : public CrossCompiler::FHlslccHeader
 	uint32 TessellationControlPointIndexBuffer;
 	EMetalOutputWindingMode TessellationOutputWinding;
 	EMetalPartitionMode TessellationPartitioning;
+	TMap<uint8, TArray<uint8>> ArgumentBuffers;
 	int8 SideTable;
 	uint8 Version;
 	bool bUsingTessellation;
@@ -728,6 +729,7 @@ static const int32 Str##PrefixLen = FCStringAnsi::Strlen(Str##Prefix)
 	DEF_PREFIX_STR(TessellationHSTFOutBuffer);
 	DEF_PREFIX_STR(TessellationControlPointOutBuffer);
 	DEF_PREFIX_STR(TessellationControlPointIndexBuffer);
+	DEF_PREFIX_STR(ArgumentBuffers);
 	DEF_PREFIX_STR(SideTable);
 #undef DEF_PREFIX_STR
 	
@@ -758,6 +760,57 @@ static const int32 Str##PrefixLen = FCStringAnsi::Strlen(Str##Prefix)
 		}
 	}
 	
+	const ANSICHAR* ArgumentTable = FCStringAnsi::Strstr(ShaderSource, ArgumentBuffersPrefix);
+	if (ArgumentTable)
+	{
+		ShaderSource = ArgumentTable;
+		ShaderSource += ArgumentBuffersPrefixLen;
+		while (*ShaderSource && *ShaderSource != '\n')
+		{
+			int32 ArgumentBufferIndex = -1;
+			if (!CrossCompiler::ParseIntegerNumber(ShaderSource, ArgumentBufferIndex))
+			{
+				return false;
+			}
+			check(ArgumentBufferIndex >= 0);
+			
+			if (!CrossCompiler::Match(ShaderSource, '['))
+			{
+				return false;
+			}
+			
+			TArray<uint8> Mask;
+			while (*ShaderSource && *ShaderSource != ']')
+			{
+				int32 MaskIndex = -1;
+				if (!CrossCompiler::ParseIntegerNumber(ShaderSource, MaskIndex))
+				{
+					return false;
+				}
+				
+				check(MaskIndex >= 0);
+				Mask.Add((uint8)MaskIndex);
+				
+				if (!CrossCompiler::Match(ShaderSource, ',') && *ShaderSource != ']')
+				{
+					return false;
+				}
+			}
+			
+			if (!CrossCompiler::Match(ShaderSource, ']'))
+			{
+				return false;
+			}
+			
+			if (!CrossCompiler::Match(ShaderSource, ',') && *ShaderSource != '\n')
+			{
+				return false;
+			}
+			
+			ArgumentBuffers.Add((uint8)ArgumentBufferIndex, Mask);
+		}
+	}
+	
 	// Early out for non-tessellation...
 	if (Version < 2 || !bUsingTessellation)
 	{
@@ -1242,6 +1295,11 @@ void BuildMetalShaderOutput(
 
 	for (auto& SamplerState : CCHeader.SamplerStates)
 	{
+		if (!SamplerMap.Contains(SamplerState.Name))
+		{
+			SamplerMap.Add(SamplerState.Name, 1);
+		}
+		
 		ParameterMap.AddParameterAllocation(
 			*SamplerState.Name,
 			0,
@@ -1271,6 +1329,12 @@ void BuildMetalShaderOutput(
 	Header.bTessFunctionConstants				= (FCStringAnsi::Strstr(USFSource, "indexBufferType [[ function_constant(32) ]]") != nullptr);
 	Header.bDeviceFunctionConstants				= (FCStringAnsi::Strstr(USFSource, "#define __METAL_DEVICE_CONSTANT_INDEX__ 1") != nullptr);
 	Header.SideTable 							= CCHeader.SideTable;
+	Header.Bindings.ArgumentBufferMasks			= CCHeader.ArgumentBuffers;
+	Header.Bindings.ArgumentBuffers				= 0;
+	for (auto const& Pair : Header.Bindings.ArgumentBufferMasks)
+	{
+		Header.Bindings.ArgumentBuffers |= (1 << Pair.Key);
+	}
 	
 	// Build the SRT for this shader.
 	{
@@ -1988,7 +2052,25 @@ void CompileShader_Metal(const FShaderCompilerInput& _Input,FShaderCompilerOutpu
 	FString MinOSVersion;
 	FString StandardVersion;
 	switch(VersionEnum)
-	{
+    {
+        case 5:
+            // Enable full SM5 feature support so tessellation & fragment UAVs compile
+            TypeMode = EMetalTypeBufferModeTex;
+            HlslCompilerTarget = HCT_FeatureLevelSM5;
+            StandardVersion = TEXT("2.1");
+            if (bAppleTV)
+            {
+                MinOSVersion = TEXT("-mtvos-version-min=12.0");
+            }
+            else if (bIsMobile)
+            {
+                MinOSVersion = TEXT("-mios-version-min=12.0");
+            }
+            else
+            {
+                MinOSVersion = TEXT("-mmacosx-version-min=10.14");
+            }
+            break;
 		case 4:
 			// Enable full SM5 feature support so tessellation & fragment UAVs compile
 			TypeMode = EMetalTypeBufferModeTex;
@@ -2162,7 +2244,10 @@ void CompileShader_Metal(const FShaderCompilerInput& _Input,FShaderCompilerOutpu
 
 
 	// This requires removing the HLSLCC_NoPreprocess flag later on!
-	RemoveUniformBuffersFromSource(Input.Environment, PreprocessedShader);
+    if (VersionEnum < 5)
+    {
+        RemoveUniformBuffersFromSource(Input.Environment, PreprocessedShader);
+    }
 
 	// Write out the preprocessed file and a batch file to compile it if requested (DumpDebugInfoPath is valid)
 	if (bDumpDebugInfo && !bDirectCompile)
diff --git a/Engine/Source/Developer/Apple/MetalShaderFormat/Private/MetalShaderFormat.cpp b/Engine/Source/Developer/Apple/MetalShaderFormat/Private/MetalShaderFormat.cpp
--- a/Engine/Source/Developer/Apple/MetalShaderFormat/Private/MetalShaderFormat.cpp
+++ b/Engine/Source/Developer/Apple/MetalShaderFormat/Private/MetalShaderFormat.cpp
@@ -239,7 +239,7 @@ class FMetalShaderFormat : public IShaderFormat
 public:
 	enum
 	{
-		HEADER_VERSION = 59,
+		HEADER_VERSION = 60,
 	};
 	
 	struct FVersion
diff --git a/Engine/Source/Developer/Apple/MetalShaderFormat/Private/MetalUtils.cpp b/Engine/Source/Developer/Apple/MetalShaderFormat/Private/MetalUtils.cpp
--- a/Engine/Source/Developer/Apple/MetalShaderFormat/Private/MetalUtils.cpp
+++ b/Engine/Source/Developer/Apple/MetalShaderFormat/Private/MetalUtils.cpp
@@ -1813,6 +1813,251 @@ static ir_dereference_variable* GenerateShaderOutput(
 
 #define USE_DS_ATTRIBUTES 0 // @todo Unicorn: consider using this as it would be more elegant, but currently isn't complete.
 
+void FMetalCodeBackend::build_iab_fields(_mesa_glsl_parse_state* ParseState, char const* n, struct glsl_type const* t, TArray<struct glsl_struct_field>& Fields, unsigned& FieldIndex, unsigned& BufferIndex, bool top, FBuffers const& Buffers)
+{
+	switch(t->base_type)
+	{
+		case GLSL_TYPE_UINT:
+		case GLSL_TYPE_INT:
+		case GLSL_TYPE_HALF:
+		case GLSL_TYPE_FLOAT:
+		case GLSL_TYPE_BOOL:
+		case GLSL_TYPE_SAMPLER:
+		case GLSL_TYPE_IMAGE:
+		case GLSL_TYPE_SAMPLER_STATE:
+		case GLSL_TYPE_ARRAY:
+		{
+			// All ignored
+			break;
+		}
+		case GLSL_TYPE_STRUCT:
+		{
+			const char *name = n;
+			for (unsigned j = 0; j < t->length; j++)
+			{
+				switch (t->fields.structure[j].type->base_type)
+				{
+					case GLSL_TYPE_STRUCT:
+					{
+						const char *newname = ralloc_asprintf(ParseState, "%s_%s",
+															  n,
+															  t->fields.structure[j].name);
+						build_iab_fields(ParseState, newname, t->fields.structure[j].type, Fields, FieldIndex, BufferIndex, false, Buffers);
+						break;
+					}
+					case GLSL_TYPE_IMAGE:
+					case GLSL_TYPE_SAMPLER:
+					case GLSL_TYPE_SAMPLER_STATE:
+					{
+						const char *newname = ralloc_asprintf(ParseState, "%s_%s",
+															  name,
+															  t->fields.structure[j].name);
+						glsl_struct_field field(t->fields.structure[j].type, newname);
+						field.semantic = ralloc_asprintf(ParseState, "[[ id(%d) ]]", FieldIndex);
+						if (t->fields.structure[j].type->sampler_buffer)
+						{
+							FieldIndex += 2;
+							BufferIndex++;
+						}
+						else
+						{
+							FieldIndex++;
+						}
+						Fields.Add(field);
+						break;
+					}
+					default:
+					{
+						break;
+					}
+				}
+			}
+			if (top)
+			{
+				if (BufferIndex)
+				{
+					glsl_struct_field field(glsl_type::GetStructuredBufferInstance("StructuredBuffer", glsl_type::uint_type), ralloc_strdup(ParseState, "BufferSizes"));
+					field.semantic = ralloc_asprintf(ParseState, "[[ id(%d) ]]", FieldIndex++);
+					field.patchconstant = 1;
+					Fields.Add(field);
+				}
+				
+				name = ralloc_asprintf(ParseState, "CB_%s", n);
+				struct glsl_type const* c = ParseState->symbols->get_type(name);
+				if (c)
+				{
+					glsl_struct_field field(c, ralloc_strdup(ParseState, n));
+					field.semantic = ralloc_asprintf(ParseState, "[[ id(%d) ]]", FieldIndex++);
+					Fields.Add(field);
+				}
+			}
+			break;
+		}
+		case GLSL_TYPE_INPUTPATCH:
+		case GLSL_TYPE_OUTPUTPATCH:
+		{
+			// Unhandled because I don't think they should appear as globals...
+			check(false);
+			break;
+		}
+		case GLSL_TYPE_VOID:
+		case GLSL_TYPE_OUTPUTSTREAM:
+		case GLSL_TYPE_ERROR:
+		case GLSL_TYPE_MAX:
+		default:
+		{
+			// Invalid types
+			check(false);
+			break;
+		}
+	}
+}
+
+struct glsl_type const* FMetalCodeBackend::create_iab_type(_mesa_glsl_parse_state* ParseState, struct glsl_type const* UBType, char const* n, FBuffers const& Buffers)
+{
+	TArray<struct glsl_struct_field> IABFields;
+	unsigned FieldIndex = 0;
+	unsigned BufferIndex = 0;
+	build_iab_fields(ParseState, n, UBType, IABFields, FieldIndex, BufferIndex, true, Buffers);
+	glsl_struct_field *const fields = ralloc_array(ParseState, glsl_struct_field, IABFields.Num());
+	memcpy(fields, IABFields.GetData(), sizeof(glsl_struct_field) * IABFields.Num());
+	const glsl_type * Result = glsl_type::get_record_instance(fields,
+												IABFields.Num(),
+												ralloc_asprintf(ParseState, "IAB_%s", n));
+	ParseState->AddUserStruct(Result);
+	return Result;
+}
+
+void FMetalCodeBackend::InsertArgumentBuffers(exec_list* ir, _mesa_glsl_parse_state* state, FBuffers& Buffers)
+{
+	ir_function_signature* EntryPointSig = GetMainFunction(ir);
+	check(EntryPointSig);
+	
+	TMap<FString, ir_variable*> IABVars;
+	foreach_iter(exec_list_iterator, Iter, EntryPointSig->parameters)
+	{
+		ir_variable* Variable = (ir_variable*)Iter.get();
+		char const* Underscore = strchr(Variable->name, '_');
+		if (Underscore)
+		{
+			uintptr_t Diff = (Underscore - Variable->name);
+			char* Name = ralloc_strndup(state, Variable->name, Diff);
+			ir_variable* UB = state->symbols->get_variable(Name);
+			if (UB)
+			{
+				ir_variable* IABVariable = IABVars.FindRef(Name);
+				if (!IABVariable)
+				{
+					struct glsl_type const* UBType = UB->type;
+					struct glsl_type const* IABType = create_iab_type(state, UBType, Name, Buffers);
+					ir_variable* NewInputVariable = new(state)ir_variable(IABType, ralloc_asprintf(state, "IAB%s", Name), ir_var_uniform);
+					NewInputVariable->semantic = ralloc_asprintf(state, "%s", Name);
+					NewInputVariable->read_only = true;
+					state->symbols->add_variable(NewInputVariable);
+					
+					IABVariable = NewInputVariable;
+					IABVars.Add(Name, NewInputVariable);
+					
+					glsl_uniform_block* block = glsl_uniform_block::alloc(state, 1);
+					block->vars[0] = NewInputVariable;
+					block->name = ralloc_asprintf(state, NewInputVariable->name);
+					const glsl_uniform_block** blocks = reralloc(state, state->uniform_blocks,
+																 const glsl_uniform_block *,
+																 state->num_uniform_blocks + 1);
+					if (blocks != NULL)
+					{
+						blocks[state->num_uniform_blocks] = block;
+						state->uniform_blocks = blocks;
+						state->num_uniform_blocks++;
+					}
+
+				}
+				
+				ir_dereference_variable* DerefVariable = new (state)ir_dereference_variable(Variable);
+				ir_dereference_record* DerefRecord = new (state)ir_dereference_record(IABVariable, ralloc_strdup(state, Variable->name));
+				int Index = IABVariable->type->field_index(DerefRecord->field);
+				if (Variable->type->base_type == GLSL_TYPE_SAMPLER)
+				{
+					IABVariable->type->fields.structure[Index].type = Variable->type;
+				}
+				ir_assignment* Assign = new (state)ir_assignment(DerefVariable, DerefRecord);
+				
+				IABVariablesMap.Add(Variable, IABVariable);
+				TSet<uint8>& Mask = IABVariableMask.FindOrAdd(IABVariable);
+				Mask.Add(Index);
+				
+				Variable->mode = ir_var_temporary;
+				Iter.remove();
+
+				for (_mesa_glsl_parse_state::TUniformList::iterator It = state->GlobalPackedArraysMap[EArrayType_Sampler].begin(); It != state->GlobalPackedArraysMap[EArrayType_Sampler].end(); ++It)
+				{
+					glsl_packed_uniform& Sampler = *It;
+					if (!strcmp(Variable->name, Sampler.Name.c_str()))
+					{
+						state->GlobalPackedArraysMap[EArrayType_Sampler].erase(It);
+						break;
+					}
+				}
+				for (_mesa_glsl_parse_state::TUniformList::iterator It = state->GlobalPackedArraysMap[EArrayType_Image].begin(); It != state->GlobalPackedArraysMap[EArrayType_Image].end(); ++It)
+				{
+					glsl_packed_uniform& Sampler = *It;
+					if (!strcmp(Variable->name, Sampler.Name.c_str()))
+					{
+						state->GlobalPackedArraysMap[EArrayType_Image].erase(It);
+						break;
+					}
+				}
+				Buffers.Remove(Variable);
+				
+				state->symbols->add_variable(Variable);
+				EntryPointSig->body.push_head(Assign);
+				EntryPointSig->body.push_head(Variable);
+			}
+		}
+	}
+	
+	foreach_iter(exec_list_iterator, Iter, EntryPointSig->parameters)
+	{
+		ir_variable* Variable = (ir_variable*)Iter.get();
+		char* Name = ralloc_strdup(state, Variable->name);
+		
+		ir_variable* UB = state->symbols->get_variable(Name);
+		if (UB && !strncmp(Variable->type->name, "CB_", 3))
+		{
+			ir_variable* IABVariable = IABVars.FindRef(Name);
+			if (IABVariable)
+			{
+				ir_dereference_variable* DerefVariable = new (state)ir_dereference_variable(Variable);
+				ir_dereference_record* DerefRecord = new (state)ir_dereference_record(IABVariable, ralloc_strdup(state, Variable->name));
+				ir_assignment* Assign = new (state)ir_assignment(DerefVariable, DerefRecord);
+				
+				Variable->mode = ir_var_ref;
+				Iter.remove();
+				
+				Buffers.Replace(Variable, IABVariable);
+				IABVariablesMap.Add(Variable, IABVariable);
+				
+				int Index = IABVariable->type->field_index(DerefRecord->field);
+				TSet<uint8>& Mask = IABVariableMask.FindOrAdd(IABVariable);
+				Mask.Add(Index);
+
+				state->symbols->add_variable(Variable);
+				EntryPointSig->body.push_head(Assign);
+				EntryPointSig->body.push_head(Variable);
+			}
+		}
+	}
+	
+	for (auto const& Pair : IABVars)
+	{
+		EntryPointSig->parameters.push_tail(Pair.Value);
+		if (Buffers.GetIndex(Pair.Value) == -1)
+		{
+			Buffers.AddBuffer(Pair.Value);
+		}
+	}
+}
+
 void FMetalCodeBackend::PackInputsAndOutputs(exec_list* Instructions, _mesa_glsl_parse_state* ParseState, EHlslShaderFrequency Frequency, exec_list& InputVars)
 {
 	ir_function_signature* EntryPointSig = GetMainFunction(Instructions);
diff --git a/Engine/Source/Developer/Apple/MetalShaderFormat/Private/MetalUtils.h b/Engine/Source/Developer/Apple/MetalShaderFormat/Private/MetalUtils.h
--- a/Engine/Source/Developer/Apple/MetalShaderFormat/Private/MetalUtils.h
+++ b/Engine/Source/Developer/Apple/MetalShaderFormat/Private/MetalUtils.h
@@ -154,6 +154,48 @@ struct FBuffers
 
 		return -1;
 	}
+	
+	void Replace(ir_variable* Old, ir_variable* New)
+	{
+		for (int i = 0, n = Buffers.Num(); i < n; ++i)
+		{
+			if (Old == Buffers[i])
+			{
+				Buffers[i] = New;
+				return;
+			}
+		}
+		
+		for (int i = 0, n = Textures.Num(); i < n; ++i)
+		{
+			if (Old == Textures[i])
+			{
+				Textures[i] = New;
+				return;
+			}
+		}
+	}
+	
+	void Remove(ir_variable* Var)
+	{
+		for (int i = 0, n = Buffers.Num(); i < n; ++i)
+		{
+			if (Var == Buffers[i])
+			{
+				Buffers[i] = nullptr;
+				return;
+			}
+		}
+		
+		for (int i = 0, n = Textures.Num(); i < n; ++i)
+		{
+			if (Var == Textures[i])
+			{
+				Textures[i] = nullptr;
+				return;
+			}
+		}
+	}
 
 	void SortBuffers(_mesa_glsl_parse_state* state)
 	{
diff --git a/Engine/Source/Developer/Mac/MacTargetPlatform/Classes/MacTargetSettings.h b/Engine/Source/Developer/Mac/MacTargetPlatform/Classes/MacTargetSettings.h
--- a/Engine/Source/Developer/Mac/MacTargetPlatform/Classes/MacTargetSettings.h
+++ b/Engine/Source/Developer/Mac/MacTargetPlatform/Classes/MacTargetSettings.h
@@ -21,6 +21,9 @@ enum class EMacMetalShaderStandard : uint8
     
     /** Metal Shaders, supporting indirect command buffers, Compatible With macOS 10.14.0 or later (std=osx-metal2.1) */
     MacMetalSLStandard_2_1 = 4 UMETA(DisplayName="Metal v2.1 (10.14.0+)"),
+	
+	/** Metal Shaders, using indirect argument buffers to reduce CPU cost, Compatible With macOS 10.14.0 or later (std=osx-metal2.1) */
+	MacMetalSLStandard_2_1_IABs = 5 UMETA(DisplayName="Metal v2.1 w/Indirect Argument Buffers (10.14.0+)"),
 };
 
 /**
diff --git a/Engine/Source/Runtime/Apple/MetalRHI/Private/MetalCommandEncoder.cpp b/Engine/Source/Runtime/Apple/MetalRHI/Private/MetalCommandEncoder.cpp
--- a/Engine/Source/Runtime/Apple/MetalRHI/Private/MetalCommandEncoder.cpp
+++ b/Engine/Source/Runtime/Apple/MetalRHI/Private/MetalCommandEncoder.cpp
@@ -1259,7 +1259,6 @@ void FMetalCommandEncoder::SetShaderBuffer(mtlpp::FunctionType const FunctionTyp
 		ShaderBuffers[uint32(FunctionType)].Lengths[(index*2)+1] = GMetalBufferFormats[Format].DataFormat;
 		ShaderBuffers[uint32(FunctionType)].Usage[index] = Usage;
 		
-		UseResource(Buffer, Usage);
 		if (FunctionType == mtlpp::FunctionType::Vertex)
 		{
 			FenceResource(Buffer);
@@ -1451,7 +1450,6 @@ void FMetalCommandEncoder::SetShaderTexture(mtlpp::FunctionType FunctionType, FM
 	
 	if (Texture)
 	{
-		UseResource(Texture, Usage);
 		TextureBindingHistory.Add(ns::AutoReleased<FMetalTexture>(Texture));
 	}
 }
@@ -1490,6 +1488,20 @@ void FMetalCommandEncoder::SetShaderSideTable(mtlpp::FunctionType const Function
 	}
 }
 
+void FMetalCommandEncoder::UseIndirectArgumentResource(FMetalTexture const& Texture, mtlpp::ResourceUsage const Usage)
+{
+	FenceResource(Texture);
+	UseResource(Texture, Usage);
+	TextureBindingHistory.Add(ns::AutoReleased<FMetalTexture>(Texture));
+}
+
+void FMetalCommandEncoder::UseIndirectArgumentResource(FMetalBuffer const& Buffer, mtlpp::ResourceUsage const Usage)
+{
+	FenceResource(Buffer);
+	UseResource(Buffer, Usage);
+	BufferBindingHistory.Add(ns::AutoReleased<FMetalBuffer>(Buffer));
+}
+
 void FMetalCommandEncoder::TransitionResources(mtlpp::Resource const& Resource)
 {
 	TransitionedResources.Add(Resource.GetPtr());
@@ -1626,8 +1638,6 @@ void FMetalCommandEncoder::SetShaderBufferInternal(mtlpp::FunctionType Function,
 				break;
 		}
 		
-		UseResource(Buffer, Usage);
-		
 		if (Buffer.IsSingleUse())
 		{
 			ShaderBuffers[uint32(Function)].Usage[Index] = mtlpp::ResourceUsage(0);
diff --git a/Engine/Source/Runtime/Apple/MetalRHI/Private/MetalCommandEncoder.h b/Engine/Source/Runtime/Apple/MetalRHI/Private/MetalCommandEncoder.h
--- a/Engine/Source/Runtime/Apple/MetalRHI/Private/MetalCommandEncoder.h
+++ b/Engine/Source/Runtime/Apple/MetalRHI/Private/MetalCommandEncoder.h
@@ -356,6 +356,20 @@ public:
 	 */
 	void SetShaderSideTable(mtlpp::FunctionType const FunctionType, NSUInteger const Index);
 	
+    /*
+     * Inform the driver that an indirect argument resource will be used and what the usage mode will be.
+     * @param Texture The texture that will be used as an indirect argument.
+     * @param Usage The usage mode for the texture.
+     */
+	void UseIndirectArgumentResource(FMetalTexture const& Texture, mtlpp::ResourceUsage const Usage);
+    
+    /*
+     * Inform the driver that an indirect argument resource will be used and what the usage mode will be.
+     * @param Buffer The buffer that will be used as an indirect argument.
+     * @param Usage The usage mode for the texture.
+     */
+	void UseIndirectArgumentResource(FMetalBuffer const& Buffer, mtlpp::ResourceUsage const Usage);
+
 	/*
 	 * Transition resource so that we can barrier fragment->vertex stages.
 	 * @param Resource The resource we are going to make readable having been written.
diff --git a/Engine/Source/Runtime/Apple/MetalRHI/Private/MetalCommandQueue.cpp b/Engine/Source/Runtime/Apple/MetalRHI/Private/MetalCommandQueue.cpp
--- a/Engine/Source/Runtime/Apple/MetalRHI/Private/MetalCommandQueue.cpp
+++ b/Engine/Source/Runtime/Apple/MetalRHI/Private/MetalCommandQueue.cpp
@@ -90,7 +90,6 @@ FMetalCommandQueue::FMetalCommandQueue(mtlpp::Device InDevice, uint32 const MaxN
 					Features |= EMetalFeaturesMaxThreadsPerThreadgroup;
 					Features |= EMetalFeaturesFences;
 					Features |= EMetalFeaturesHeaps;
-					Features |= EMetalFeaturesIABs;
 					
 					if (MaxShaderVersion >= 4)
 					{
@@ -151,7 +150,6 @@ FMetalCommandQueue::FMetalCommandQueue(mtlpp::Device InDevice, uint32 const MaxN
 				Features |= EMetalFeaturesMaxThreadsPerThreadgroup;
 				Features |= EMetalFeaturesFences;
 				Features |= EMetalFeaturesHeaps;
-				Features |= EMetalFeaturesIABs;
 				
 				if (MaxShaderVersion >= 4)
 				{
@@ -225,11 +223,15 @@ FMetalCommandQueue::FMetalCommandQueue(mtlpp::Device InDevice, uint32 const MaxN
 		// Turn on Texture Buffers! These are faster on the GPU as we don't need to do out-of-bounds tests but require Metal 2.1 and macOS 10.14
 		if (FPlatformMisc::MacOSXVersionCompare(10,14,0) >= 0)
 		{
-			Features |= EMetalFeaturesMaxThreadsPerThreadgroup | EMetalFeaturesIABs;
+			Features |= EMetalFeaturesMaxThreadsPerThreadgroup;
 			if (MaxShaderVersion >= 4)
 			{
 				Features |= EMetalFeaturesTextureBuffers;
-			}
+            }
+            if (MaxShaderVersion >= 5)
+            {
+                Features |= EMetalFeaturesIABs;
+            }
 			
 			if (!FParse::Param(FCommandLine::Get(),TEXT("nometalfence")))
 			{
diff --git a/Engine/Source/Runtime/Apple/MetalRHI/Private/MetalCommands.cpp b/Engine/Source/Runtime/Apple/MetalRHI/Private/MetalCommands.cpp
--- a/Engine/Source/Runtime/Apple/MetalRHI/Private/MetalCommands.cpp
+++ b/Engine/Source/Runtime/Apple/MetalRHI/Private/MetalCommands.cpp
@@ -536,7 +536,7 @@ void FMetalRHICommandContext::RHISetShaderUniformBuffer(FVertexShaderRHIParamRef
 	if (Bindings.ConstantBuffers & 1 << BufferIndex)
 	{
 		auto* UB = (FMetalUniformBuffer*)BufferRHI;
-		Context->GetCurrentState().SetShaderBuffer(SF_Vertex, UB->Buffer, UB->Data, 0, UB->GetSize(), BufferIndex, mtlpp::ResourceUsage::Read);
+		Context->GetCurrentState().SetShaderBuffer(SF_Vertex, (Bindings.ArgumentBuffers & 1 << BufferIndex) ? UB->IndirectArgumentBuffer : UB->Buffer, UB->Data, 0, UB->GetSize(), BufferIndex, mtlpp::ResourceUsage::Read);
 	}
 	}
 }
@@ -552,7 +552,7 @@ void FMetalRHICommandContext::RHISetShaderUniformBuffer(FHullShaderRHIParamRef H
 	if (Bindings.ConstantBuffers & 1 << BufferIndex)
 	{
 		auto* UB = (FMetalUniformBuffer*)BufferRHI;
-		Context->GetCurrentState().SetShaderBuffer(SF_Hull, UB->Buffer, UB->Data, 0, UB->GetSize(), BufferIndex, mtlpp::ResourceUsage::Read);
+		Context->GetCurrentState().SetShaderBuffer(SF_Hull, (Bindings.ArgumentBuffers & 1 << BufferIndex) ? UB->IndirectArgumentBuffer : UB->Buffer, UB->Data, 0, UB->GetSize(), BufferIndex, mtlpp::ResourceUsage::Read);
 	}
 	}
 }
@@ -568,7 +568,7 @@ void FMetalRHICommandContext::RHISetShaderUniformBuffer(FDomainShaderRHIParamRef
 	if (Bindings.ConstantBuffers & 1 << BufferIndex)
 	{
 		auto* UB = (FMetalUniformBuffer*)BufferRHI;
-		Context->GetCurrentState().SetShaderBuffer(SF_Domain, UB->Buffer, UB->Data, 0, UB->GetSize(), BufferIndex, mtlpp::ResourceUsage::Read);
+		Context->GetCurrentState().SetShaderBuffer(SF_Domain, (Bindings.ArgumentBuffers & 1 << BufferIndex) ? UB->IndirectArgumentBuffer : UB->Buffer, UB->Data, 0, UB->GetSize(), BufferIndex, mtlpp::ResourceUsage::Read);
 	}
 	}
 }
@@ -589,7 +589,7 @@ void FMetalRHICommandContext::RHISetShaderUniformBuffer(FPixelShaderRHIParamRef
 	if (Bindings.ConstantBuffers & 1 << BufferIndex)
 	{
 		auto* UB = (FMetalUniformBuffer*)BufferRHI;
-		Context->GetCurrentState().SetShaderBuffer(SF_Pixel, UB->Buffer, UB->Data, 0, UB->GetSize(), BufferIndex, mtlpp::ResourceUsage::Read);
+		Context->GetCurrentState().SetShaderBuffer(SF_Pixel, (Bindings.ArgumentBuffers & 1 << BufferIndex) ? UB->IndirectArgumentBuffer : UB->Buffer, UB->Data, 0, UB->GetSize(), BufferIndex, mtlpp::ResourceUsage::Read);
 	}
 	}
 }
@@ -605,7 +605,7 @@ void FMetalRHICommandContext::RHISetShaderUniformBuffer(FComputeShaderRHIParamRe
 	if (Bindings.ConstantBuffers & 1 << BufferIndex)
 	{
 		auto* UB = (FMetalUniformBuffer*)BufferRHI;
-		Context->GetCurrentState().SetShaderBuffer(SF_Compute, UB->Buffer, UB->Data, 0, UB->GetSize(), BufferIndex, mtlpp::ResourceUsage::Read);
+		Context->GetCurrentState().SetShaderBuffer(SF_Compute, (Bindings.ArgumentBuffers & 1 << BufferIndex) ? UB->IndirectArgumentBuffer : UB->Buffer, UB->Data, 0, UB->GetSize(), BufferIndex, mtlpp::ResourceUsage::Read);
 	}
 	}
 }
diff --git a/Engine/Source/Runtime/Apple/MetalRHI/Private/MetalIndexBuffer.cpp b/Engine/Source/Runtime/Apple/MetalRHI/Private/MetalIndexBuffer.cpp
--- a/Engine/Source/Runtime/Apple/MetalRHI/Private/MetalIndexBuffer.cpp
+++ b/Engine/Source/Runtime/Apple/MetalRHI/Private/MetalIndexBuffer.cpp
@@ -47,7 +47,7 @@ FIndexBufferRHIRef FMetalDynamicRHI::RHICreateIndexBuffer(uint32 Stride,uint32 S
 	}
 	else if (IndexBuffer->Buffer.GetStorageMode() == mtlpp::StorageMode::Private)
 	{
-		if (IndexBuffer->GetUsage() & (BUF_Dynamic|BUF_Static))
+		if (IndexBuffer->UsePrivateMemory())
 		{
 			LLM_SCOPE(ELLMTag::IndexBuffer);
 			SafeReleaseMetalBuffer(IndexBuffer->CPUBuffer);
@@ -106,7 +106,7 @@ struct FMetalRHICommandInitialiseIndexBuffer : public FRHICommand<FMetalRHIComma
 		if (Buffer->CPUBuffer)
 		{
 			GetMetalDeviceContext().AsyncCopyFromBufferToBuffer(Buffer->CPUBuffer, 0, Buffer->Buffer, 0, Buffer->Buffer.GetLength());
-			if (Buffer->GetUsage() & (BUF_Dynamic|BUF_Static))
+			if (Buffer->UsePrivateMemory())
 			{
 				LLM_SCOPE(ELLMTag::IndexBuffer);
 				SafeReleaseMetalBuffer(Buffer->CPUBuffer);
@@ -170,7 +170,7 @@ FIndexBufferRHIRef FMetalDynamicRHI::CreateIndexBuffer_RenderThread(class FRHICo
 		}
 		else if (IndexBuffer->Buffer.GetStorageMode() == mtlpp::StorageMode::Private)
 		{
-			if (IndexBuffer->GetUsage() & (BUF_Dynamic|BUF_Static))
+			if (IndexBuffer->UsePrivateMemory())
 			{
 				LLM_SCOPE(ELLMTag::IndexBuffer);
 				SafeReleaseMetalBuffer(IndexBuffer->CPUBuffer);
diff --git a/Engine/Source/Runtime/Apple/MetalRHI/Private/MetalState.cpp b/Engine/Source/Runtime/Apple/MetalRHI/Private/MetalState.cpp
--- a/Engine/Source/Runtime/Apple/MetalRHI/Private/MetalState.cpp
+++ b/Engine/Source/Runtime/Apple/MetalRHI/Private/MetalState.cpp
@@ -345,6 +345,10 @@ static FMetalSampler FindOrCreateSamplerState(mtlpp::Device Device, const FSampl
 #if PLATFORM_MAC
 		Desc.SetBorderColor(Initializer.BorderColor == 0 ? mtlpp::SamplerBorderColor::TransparentBlack : mtlpp::SamplerBorderColor::OpaqueWhite);
 #endif
+		if (FMetalCommandQueue::SupportsFeature(EMetalFeaturesIABs))
+		{
+			Desc.SetSupportArgumentBuffers(true);
+		}
 		
 		State = Device.NewSamplerState(Desc);
 		
diff --git a/Engine/Source/Runtime/Apple/MetalRHI/Private/MetalStateCache.cpp b/Engine/Source/Runtime/Apple/MetalRHI/Private/MetalStateCache.cpp
--- a/Engine/Source/Runtime/Apple/MetalRHI/Private/MetalStateCache.cpp
+++ b/Engine/Source/Runtime/Apple/MetalRHI/Private/MetalStateCache.cpp
@@ -2221,6 +2221,67 @@ void FMetalStateCache::CommitResourceTable(EShaderFrequency const Frequency, mtl
 			CommandEncoder.SetShaderSamplerState(Type, SamplerBindings.Samplers[Index], Index);
 		}
 	}
+	
+	uint32 ArgumentBuffers = 0;
+	TMap<uint8, TArray<uint8>> ArgumentBufferMasks;
+	switch(Frequency)
+	{
+		case SF_Vertex:
+			ArgumentBuffers = GraphicsPSO->VertexShader->Bindings.ArgumentBuffers;
+			ArgumentBufferMasks = GraphicsPSO->VertexShader->Bindings.ArgumentBufferMasks;
+			break;
+		case SF_Hull:
+			ArgumentBuffers = GraphicsPSO->HullShader->Bindings.ArgumentBuffers;
+			ArgumentBufferMasks = GraphicsPSO->HullShader->Bindings.ArgumentBufferMasks;
+			break;
+		case SF_Domain:
+			ArgumentBuffers = GraphicsPSO->DomainShader->Bindings.ArgumentBuffers;
+			ArgumentBufferMasks = GraphicsPSO->DomainShader->Bindings.ArgumentBufferMasks;
+			break;
+		case SF_Pixel:
+			ArgumentBuffers = GraphicsPSO->PixelShader->Bindings.ArgumentBuffers;
+			ArgumentBufferMasks = GraphicsPSO->PixelShader->Bindings.ArgumentBufferMasks;
+			break;
+		case SF_Compute:
+			ArgumentBuffers = ComputeShader->Bindings.ArgumentBuffers;
+			ArgumentBufferMasks = ComputeShader->Bindings.ArgumentBufferMasks;
+			break;
+		default:
+			break;
+	}
+	
+	while(ArgumentBuffers)
+	{
+		uint32 Index = __builtin_ctz(ArgumentBuffers);
+		ArgumentBuffers &= ~(1 << Index);
+		
+		FMetalUniformBuffer* Buffer = (FMetalUniformBuffer*)BoundUniformBuffers[Frequency][Index];
+		check(Buffer);
+		
+		TArray<uint8>& Mask = ArgumentBufferMasks.FindChecked((uint8)Index);
+		for (uint8 Entry : Mask)
+		{
+			FMetalUniformBuffer::Argument const& Resource = Buffer->IndirectArgumentResources[Entry];
+			if (Resource.Buffer)
+			{
+				CommandEncoder.UseIndirectArgumentResource(Resource.Buffer, Resource.Usage);
+			}
+			else if (Resource.Texture)
+			{
+				CommandEncoder.UseIndirectArgumentResource(Resource.Texture, Resource.Usage);
+			}
+		}
+		
+		if (Buffer->IndirectArgumentBufferSideTable)
+		{
+			CommandEncoder.UseIndirectArgumentResource(Buffer->IndirectArgumentBufferSideTable, mtlpp::ResourceUsage::Read);
+		}
+		
+		if (Buffer->Buffer)
+		{
+			CommandEncoder.UseIndirectArgumentResource(Buffer->Buffer, mtlpp::ResourceUsage::Read);
+		}
+	}
 }
 
 FTexture2DRHIRef FMetalStateCache::CreateFallbackDepthStencilSurface(uint32 Width, uint32 Height)
diff --git a/Engine/Source/Runtime/Apple/MetalRHI/Private/MetalUniformBuffer.cpp b/Engine/Source/Runtime/Apple/MetalRHI/Private/MetalUniformBuffer.cpp
--- a/Engine/Source/Runtime/Apple/MetalRHI/Private/MetalUniformBuffer.cpp
+++ b/Engine/Source/Runtime/Apple/MetalRHI/Private/MetalUniformBuffer.cpp
@@ -9,6 +9,7 @@
 #include "MetalBuffer.h"
 #include "MetalCommandBuffer.h"
 #include "HAL/LowLevelMemTracker.h"
+#include "Misc/ScopeRWLock.h"
 
 struct FMetalRHICommandInitialiseUniformdBuffer : public FRHICommand<FMetalRHICommandInitialiseUniformdBuffer>
 {
@@ -38,9 +39,68 @@ struct FMetalRHICommandInitialiseUniformdBuffer : public FRHICommand<FMetalRHICo
 	}
 };
 
+struct FMetalArgumentEncoderMapFuncs : TDefaultMapKeyFuncs<ns::Array<mtlpp::ArgumentDescriptor>, mtlpp::ArgumentEncoder, /*bInAllowDuplicateKeys*/false>
+{
+	static FORCEINLINE uint32 GetKeyHash(KeyInitType Key)
+	{
+		uint32 Hash = 0;
+		
+		for (ns::AutoReleased<mtlpp::ArgumentDescriptor> const& Desc : Key)
+		{
+			Hash ^= ((uint32)Desc.GetDataType() * (uint32)Desc.GetTextureType() * (uint32)Desc.GetAccess() * Desc.GetArrayLength()) << Desc.GetIndex();
+		}
+		
+		return Hash;
+	}
+};
+
+class FMetalArgumentEncoderCache
+{
+public:
+	FMetalArgumentEncoderCache()
+	{
+	}
+	
+	~FMetalArgumentEncoderCache()
+	{
+	}
+	
+	static FMetalArgumentEncoderCache& Get()
+	{
+		static FMetalArgumentEncoderCache sSelf;
+		return sSelf;
+	}
+	
+	mtlpp::ArgumentEncoder CreateEncoder(ns::Array<mtlpp::ArgumentDescriptor> const& Desc)
+	{
+		mtlpp::ArgumentEncoder Encoder;
+		
+		FRWScopeLock Lock(Mutex, SLT_ReadOnly);
+		Encoder = Encoders.FindRef(Desc);
+		
+		if (!Encoder)
+		{
+			Encoder = GetMetalDeviceContext().GetDevice().NewArgumentEncoderWithArguments(Desc);
+			
+			// Now we are a writer as we want to create & add the new pipeline
+			Lock.ReleaseReadOnlyLockAndAcquireWriteLock_USE_WITH_CAUTION();
+			
+			if (!Encoders.Find(Desc))
+			{
+				Encoders.Add(Desc, Encoder);
+			}
+		}
+		
+		return Encoder;
+	}
+private:
+	FRWLock Mutex;
+	TMap<ns::Array<mtlpp::ArgumentDescriptor>, mtlpp::ArgumentEncoder, FDefaultSetAllocator, FMetalArgumentEncoderMapFuncs> Encoders;
+};
+
 FMetalUniformBuffer::FMetalUniformBuffer(const void* Contents, const FRHIUniformBufferLayout& Layout, EUniformBufferUsage InUsage)
 	: FRHIUniformBuffer(Layout)
-	, FMetalRHIBuffer(Layout.ConstantBufferSize, BUF_Volatile, RRT_UniformBuffer)
+    , FMetalRHIBuffer(Layout.ConstantBufferSize, (FMetalCommandQueue::SupportsFeature(EMetalFeaturesIABs) && Layout.Resources.Num() ? (EMetalBufferUsage_GPUOnly|BUF_Volatile) : BUF_Volatile), RRT_UniformBuffer)
 {
 	if (Layout.ConstantBufferSize > 0)
 	{
@@ -71,25 +131,298 @@ FMetalUniformBuffer::FMetalUniformBuffer(const void* Contents, const FRHIUniform
 	// set up an SRT-style uniform buffer
 	if (Layout.Resources.Num())
 	{
-		int32 NumResources = Layout.Resources.Num();
-		ResourceTable.Empty(NumResources);
-		ResourceTable.AddZeroed(NumResources);
-		for (int32 i = 0; i < NumResources; ++i)
+		if (FMetalCommandQueue::SupportsFeature(EMetalFeaturesIABs))
 		{
-			FRHIResource* Resource = *(FRHIResource**)((uint8*)Contents + Layout.ResourceOffsets[i]);
-
-			// Allow null SRV's in uniform buffers for feature levels that don't support SRV's in shaders
-			if (!(GMaxRHIFeatureLevel <= ERHIFeatureLevel::ES3_1 && Layout.Resources[i] == UBMT_SRV))
+			TArray<id> Pointers;
+			TArray<uint32> BufferSizes;
+			NSMutableArray<mtlpp::ArgumentDescriptor::Type>* Arguments = [[NSMutableArray new] autorelease];
+			
+			int32 NumResources = Layout.Resources.Num();
+			ResourceTable.Empty(NumResources);
+			ResourceTable.AddZeroed(NumResources);
+			int32 Index = 0;
+			for (int32 i = 0; i < NumResources; ++i, ++Index)
+			{
+				FRHIResource* Resource = *(FRHIResource**)((uint8*)Contents + Layout.ResourceOffsets[i]);
+				
+				// Allow null SRV's in uniform buffers for feature levels that don't support SRV's in shaders
+				if (!(GMaxRHIFeatureLevel <= ERHIFeatureLevel::ES3_1 && Layout.Resources[i] == UBMT_SRV))
+				{
+					check(Resource);
+				}
+				ResourceTable[i] = Resource;
+				
+				switch(Layout.Resources[i])
+				{
+					case UBMT_GRAPH_TRACKED_SRV:
+					case UBMT_GRAPH_TRACKED_BUFFER_SRV:
+					case UBMT_SRV:
+					{
+						mtlpp::ArgumentDescriptor Desc;
+						Desc.SetIndex(Index);
+						Desc.SetAccess(mtlpp::ArgumentAccess::ReadOnly);
+						[Arguments addObject:Desc.GetPtr()];
+						
+						FMetalShaderResourceView* SRV = (FMetalShaderResourceView*)Resource;
+						FRHITexture* Texture = SRV->SourceTexture.GetReference();
+						FMetalVertexBuffer* VB = SRV->SourceVertexBuffer.GetReference();
+						FMetalIndexBuffer* IB = SRV->SourceIndexBuffer.GetReference();
+						FMetalStructuredBuffer* SB = SRV->SourceStructuredBuffer.GetReference();
+						if (Texture)
+						{
+							FMetalSurface* Surface = SRV->TextureView;
+							check (Surface != nullptr);
+							Desc.SetDataType(mtlpp::DataType::Texture);
+							Desc.SetTextureType(Surface->Texture.GetTextureType());
+							
+							check(!Surface->Texture.IsAliasable());
+							IndirectArgumentResources.Add(Argument(Surface->Texture, (mtlpp::ResourceUsage)(mtlpp::ResourceUsage::Read|mtlpp::ResourceUsage::Sample)));
+							Pointers.Add(Surface->Texture.GetPtr());
+						}
+						else
+						{
+							check(VB || IB || SB);
+							ns::AutoReleased<FMetalTexture> Tex = SRV->GetLinearTexture(false);
+							Desc.SetDataType(mtlpp::DataType::Texture);
+							Desc.SetTextureType(Tex.GetTextureType());
+							Pointers.Add(Tex.GetPtr());
+							IndirectArgumentResources.Add(Argument(Tex, (mtlpp::ResourceUsage)(mtlpp::ResourceUsage::Read|mtlpp::ResourceUsage::Sample)));
+							
+							mtlpp::ArgumentDescriptor BufferDesc;
+							BufferDesc.SetIndex(++Index);
+							BufferDesc.SetAccess(mtlpp::ArgumentAccess::ReadOnly);
+							[Arguments addObject:BufferDesc.GetPtr()];
+							
+							if (VB)
+							{
+								BufferDesc.SetDataType(mtlpp::DataType::Pointer);
+								Pointers.Add(VB->Buffer.GetPtr());
+								IndirectArgumentResources.Add(Argument(VB->Buffer, (mtlpp::ResourceUsage)(mtlpp::ResourceUsage::Read)));
+								
+								check(VB->Buffer.GetStorageMode() == mtlpp::StorageMode::Private);
+								BufferSizes.Add(VB->GetSize());
+								BufferSizes.Add(GMetalBufferFormats[SRV->Format].DataFormat);
+							}
+							else if (IB)
+							{
+								BufferDesc.SetDataType(mtlpp::DataType::Pointer);
+								Pointers.Add(IB->Buffer.GetPtr());
+								IndirectArgumentResources.Add(Argument(IB->Buffer, (mtlpp::ResourceUsage)(mtlpp::ResourceUsage::Read)));
+								
+								check(IB->Buffer.GetStorageMode() == mtlpp::StorageMode::Private);
+								BufferSizes.Add(IB->GetSize());
+								BufferSizes.Add(GMetalBufferFormats[SRV->Format].DataFormat);
+							}
+							else if (SB)
+							{
+								BufferDesc.SetDataType(mtlpp::DataType::Pointer);
+								Pointers.Add(SB->Buffer.GetPtr());
+								IndirectArgumentResources.Add(Argument(SB->Buffer, (mtlpp::ResourceUsage)(mtlpp::ResourceUsage::Read)));
+								
+								BufferSizes.Add(SB->GetSize());
+								BufferSizes.Add(GMetalBufferFormats[SRV->Format].DataFormat);
+							}
+						}
+						break;
+					}
+					case UBMT_GRAPH_TRACKED_UAV:
+					case UBMT_GRAPH_TRACKED_BUFFER_UAV:
+					{
+						mtlpp::ArgumentDescriptor Desc;
+						Desc.SetIndex(Index);
+						Desc.SetAccess(mtlpp::ArgumentAccess::ReadWrite);
+						[Arguments addObject:Desc.GetPtr()];
+						
+						FMetalUnorderedAccessView* UAV = (FMetalUnorderedAccessView*)Resource;
+						FMetalShaderResourceView* SRV = (FMetalShaderResourceView*)UAV->SourceView;
+						FMetalStructuredBuffer* SB = UAV->SourceView->SourceStructuredBuffer.GetReference();
+						FMetalVertexBuffer* VB = UAV->SourceView->SourceVertexBuffer.GetReference();
+						FMetalIndexBuffer* IB = UAV->SourceView->SourceIndexBuffer.GetReference();
+						FRHITexture* Texture = UAV->SourceView->SourceTexture.GetReference();
+						FMetalSurface* Surface = UAV->SourceView->TextureView;
+						if (Texture)
+						{
+							if (!Surface)
+							{
+								Surface = GetMetalSurfaceFromRHITexture(Texture);
+							}
+							check (Surface != nullptr);
+							Desc.SetDataType(mtlpp::DataType::Texture);
+							Desc.SetTextureType(Surface->Texture.GetTextureType());
+							
+							check(!Surface->Texture.IsAliasable());
+							Pointers.Add(Surface->Texture.GetPtr());
+							IndirectArgumentResources.Add(Argument(Surface->Texture, (mtlpp::ResourceUsage)(mtlpp::ResourceUsage::Read|mtlpp::ResourceUsage::Write)));
+						}
+						else
+						{
+							check(VB || IB || SB);
+							ns::AutoReleased<FMetalTexture> Tex = SRV->GetLinearTexture(false);
+							Desc.SetDataType(mtlpp::DataType::Texture);
+							Desc.SetTextureType(Tex.GetTextureType());
+							Pointers.Add(Tex.GetPtr());
+							IndirectArgumentResources.Add(Argument(Tex, (mtlpp::ResourceUsage)(mtlpp::ResourceUsage::Read|mtlpp::ResourceUsage::Write)));
+							
+							mtlpp::ArgumentDescriptor BufferDesc;
+							BufferDesc.SetIndex(++Index);
+							BufferDesc.SetAccess(mtlpp::ArgumentAccess::ReadWrite);
+							[Arguments addObject:BufferDesc.GetPtr()];
+							
+							if (VB)
+							{
+								BufferDesc.SetDataType(mtlpp::DataType::Pointer);
+								Pointers.Add(VB->Buffer.GetPtr());
+								IndirectArgumentResources.Add(Argument(VB->Buffer, (mtlpp::ResourceUsage)(mtlpp::ResourceUsage::Read|mtlpp::ResourceUsage::Write)));
+								
+								check(VB->Buffer.GetStorageMode() == mtlpp::StorageMode::Private);
+								BufferSizes.Add(VB->GetSize());
+								BufferSizes.Add(GMetalBufferFormats[SRV->Format].DataFormat);
+							}
+							else if (IB)
+							{
+								BufferDesc.SetDataType(mtlpp::DataType::Pointer);
+								Pointers.Add(IB->Buffer.GetPtr());
+								IndirectArgumentResources.Add(Argument(IB->Buffer, (mtlpp::ResourceUsage)(mtlpp::ResourceUsage::Read|mtlpp::ResourceUsage::Write)));
+								
+								check(IB->Buffer.GetStorageMode() == mtlpp::StorageMode::Private);
+								BufferSizes.Add(IB->GetSize());
+								BufferSizes.Add(GMetalBufferFormats[SRV->Format].DataFormat);
+							}
+							else if (SB)
+							{
+								BufferDesc.SetDataType(mtlpp::DataType::Pointer);
+								Pointers.Add(SB->Buffer.GetPtr());
+								IndirectArgumentResources.Add(Argument(SB->Buffer, (mtlpp::ResourceUsage)(mtlpp::ResourceUsage::Read|mtlpp::ResourceUsage::Write)));
+								
+								BufferSizes.Add(SB->GetSize());
+								BufferSizes.Add(GMetalBufferFormats[SRV->Format].DataFormat);
+							}
+						}
+						break;
+					}
+					case UBMT_SAMPLER:
+					{
+						mtlpp::ArgumentDescriptor Desc;
+						Desc.SetIndex(Index);
+						Desc.SetAccess(mtlpp::ArgumentAccess::ReadOnly);
+						[Arguments addObject:Desc.GetPtr()];
+						
+						FMetalSamplerState* Sampler = (FMetalSamplerState*)Resource;
+						Desc.SetDataType(mtlpp::DataType::Sampler);
+						Pointers.Add(Sampler->State.GetPtr());
+						IndirectArgumentResources.Add(Argument());
+						break;
+					}
+					case UBMT_GRAPH_TRACKED_TEXTURE:
+					case UBMT_TEXTURE:
+					{
+						mtlpp::ArgumentDescriptor Desc;
+						Desc.SetIndex(Index);
+						Desc.SetAccess(mtlpp::ArgumentAccess::ReadOnly);
+						[Arguments addObject:Desc.GetPtr()];
+						
+						FRHITexture* Texture = (FRHITexture*)Resource;
+						FMetalSurface* Surface = GetMetalSurfaceFromRHITexture(Texture);
+						check (Surface != nullptr);
+						Desc.SetDataType(mtlpp::DataType::Texture);
+						Desc.SetTextureType(Surface->Texture.GetTextureType());
+						
+						check(!Surface->Texture.IsAliasable());
+						Pointers.Add(Surface->Texture.GetPtr());
+						IndirectArgumentResources.Add(Argument(Surface->Texture, (mtlpp::ResourceUsage)(mtlpp::ResourceUsage::Read|mtlpp::ResourceUsage::Sample)));
+						break;
+					}
+					default:
+						break;
+				}
+			}
+			
+			if (BufferSizes.Num() > 0)
+			{
+				mtlpp::ArgumentDescriptor Desc;
+				Desc.SetIndex(Index++);
+				Desc.SetAccess(mtlpp::ArgumentAccess::ReadOnly);
+				Desc.SetDataType(mtlpp::DataType::Pointer);
+				[Arguments addObject:Desc.GetPtr()];
+				
+				FMetalPooledBufferArgs Args(GetMetalDeviceContext().GetDevice(), BufferSizes.Num() * sizeof(uint32), BUFFER_STORAGE_MODE);
+				IndirectArgumentBufferSideTable = GetMetalDeviceContext().CreatePooledBuffer(Args);
+				
+				FMemory::Memcpy(IndirectArgumentBufferSideTable.GetContents(), BufferSizes.GetData(), BufferSizes.Num() * sizeof(uint32));
+				
+				IndirectArgumentResources.Add(Argument(IndirectArgumentBufferSideTable, mtlpp::ResourceUsage::Read));
+				Pointers.Add(IndirectArgumentBufferSideTable.GetPtr());
+			}
+			
+			if (Layout.ConstantBufferSize > 0)
+            {
+                mtlpp::ArgumentDescriptor Desc;
+                Desc.SetIndex(Index++);
+                Desc.SetAccess(mtlpp::ArgumentAccess::ReadOnly);
+                Desc.SetDataType(mtlpp::DataType::Pointer);
+                [Arguments addObject:Desc.GetPtr()];
+                
+                Pointers.Add(Buffer.GetPtr());
+                IndirectArgumentResources.Add(Argument(Buffer, mtlpp::ResourceUsage::Read));
+			}
+			
+			ns::Array<mtlpp::ArgumentDescriptor> ArgArray(Arguments);
+			mtlpp::ArgumentEncoder Encoder = FMetalArgumentEncoderCache::Get().CreateEncoder(ArgArray);
+			
+			FMetalPooledBufferArgs Args(GetMetalDeviceContext().GetDevice(), Encoder.GetEncodedLength(), BUFFER_STORAGE_MODE);
+			IndirectArgumentBuffer = GetMetalDeviceContext().CreatePooledBuffer(Args);
+			Encoder.SetArgumentBuffer(IndirectArgumentBuffer, 0);
+			
+			for (MTLArgumentDescriptor* Arg in Arguments)
+			{
+				mtlpp::ArgumentDescriptor Argument(Arg);
+				NSUInteger NewIndex = Argument.GetIndex();
+				switch(Argument.GetDataType())
+				{
+					case mtlpp::DataType::Pointer:
+						Encoder.SetBuffer((id<MTLBuffer>)Pointers[NewIndex], 0, NewIndex);
+						break;
+					case mtlpp::DataType::Texture:
+						Encoder.SetTexture((id<MTLTexture>)Pointers[NewIndex], NewIndex);
+						break;
+					case mtlpp::DataType::Sampler:
+						Encoder.SetSamplerState((id<MTLSamplerState>)Pointers[NewIndex], NewIndex);
+						break;
+					case mtlpp::DataType::Struct:
+					{
+						void* Data = Encoder.GetConstantDataAtIndex(NewIndex);
+						FMemory::Memcpy(Data, Contents, Layout.ConstantBufferSize);
+						break;
+					}
+					default:
+						break;
+				}
+			}
+		}
+		else
+		{
+			int32 NumResources = Layout.Resources.Num();
+			ResourceTable.Empty(NumResources);
+			ResourceTable.AddZeroed(NumResources);
+			for (int32 i = 0; i < NumResources; ++i)
 			{
-				check(Resource);
+				FRHIResource* Resource = *(FRHIResource**)((uint8*)Contents + Layout.ResourceOffsets[i]);
+				
+				// Allow null SRV's in uniform buffers for feature levels that don't support SRV's in shaders
+				if (!(GMaxRHIFeatureLevel <= ERHIFeatureLevel::ES3_1 && Layout.Resources[i] == UBMT_SRV))
+				{
+					check(Resource);
+				}
+				ResourceTable[i] = Resource;
 			}
-			ResourceTable[i] = Resource;
 		}
 	}
 }
 
 FMetalUniformBuffer::~FMetalUniformBuffer()
 {
+	SafeReleaseMetalBuffer(IndirectArgumentBuffer);
+	SafeReleaseMetalBuffer(IndirectArgumentBufferSideTable);
 }
 
 void const* FMetalUniformBuffer::GetData()
diff --git a/Engine/Source/Runtime/Apple/MetalRHI/Private/MetalVertexBuffer.cpp b/Engine/Source/Runtime/Apple/MetalRHI/Private/MetalVertexBuffer.cpp
--- a/Engine/Source/Runtime/Apple/MetalRHI/Private/MetalVertexBuffer.cpp
+++ b/Engine/Source/Runtime/Apple/MetalRHI/Private/MetalVertexBuffer.cpp
@@ -101,6 +101,12 @@ FMetalVertexBuffer::~FMetalVertexBuffer()
 {
 }
 
+bool FMetalRHIBuffer::UsePrivateMemory() const
+{
+	return (FMetalCommandQueue::SupportsFeature(EMetalFeaturesEfficientBufferBlits) && (Usage & (BUF_Dynamic|BUF_Static)))
+	|| (FMetalCommandQueue::SupportsFeature(EMetalFeaturesIABs) && Usage & (BUF_ShaderResource|BUF_UnorderedAccess));
+}
+
 FMetalRHIBuffer::FMetalRHIBuffer(uint32 InSize, uint32 InUsage, ERHIResourceType InType)
 : Data(nullptr)
 , LastUpdate(0)
@@ -230,7 +236,7 @@ void FMetalRHIBuffer::Unalias()
 void FMetalRHIBuffer::Alloc(uint32 InSize, EResourceLockMode LockMode)
 {
 	METAL_LLM_BUFFER_SCOPE(Type);
-	bool const bUsePrivateMem = (Usage & (BUF_Static|BUF_Dynamic)) && FMetalCommandQueue::SupportsFeature(EMetalFeaturesEfficientBufferBlits);
+	bool const bUsePrivateMem = UsePrivateMemory();
 
 	if (!Buffer)
 	{
@@ -420,7 +426,7 @@ void* FMetalRHIBuffer::Lock(EResourceLockMode LockMode, uint32 Offset, uint32 In
 	// In order to properly synchronise the buffer access, when a dynamic buffer is locked for writing, discard the old buffer & create a new one. This prevents writing to a buffer while it is being read by the GPU & thus causing corruption. This matches the logic of other RHIs.
 	if (LockMode == RLM_WriteOnly)
 	{
-        bool const bUsePrivateMem = (Usage & (BUF_Static|BUF_Dynamic)) && FMetalCommandQueue::SupportsFeature(EMetalFeaturesEfficientBufferBlits);
+        bool const bUsePrivateMem = UsePrivateMemory();
         if (bUsePrivateMem)
         {
 			METAL_LLM_BUFFER_SCOPE(Type);
@@ -511,7 +517,7 @@ void FMetalRHIBuffer::Unlock()
 
 			// Synchronise the buffer with the GPU
 			GetMetalDeviceContext().AsyncCopyFromBufferToBuffer(CPUBuffer, 0, Buffer, 0, Buffer.GetLength());
-			if (Usage & (BUF_Dynamic|BUF_Static))
+			if (UsePrivateMemory())
             {
 				METAL_LLM_BUFFER_SCOPE(Type);
 				SafeReleaseMetalBuffer(CPUBuffer);
@@ -559,7 +565,7 @@ FVertexBufferRHIRef FMetalDynamicRHI::RHICreateVertexBuffer(uint32 Size, uint32
 	}
 	else if (VertexBuffer->Buffer.GetStorageMode() == mtlpp::StorageMode::Private)
 	{
-		if (VertexBuffer->GetUsage() & (BUF_Dynamic|BUF_Static))
+		if (VertexBuffer->UsePrivateMemory())
 		{
 			LLM_SCOPE(ELLMTag::VertexBuffer);
 			SafeReleaseMetalBuffer(VertexBuffer->CPUBuffer);
@@ -642,7 +648,7 @@ struct FMetalRHICommandInitialiseVertexBuffer : public FRHICommand<FMetalRHIComm
 			uint32 Size = FMath::Min(Buffer->Buffer.GetLength(), Buffer->CPUBuffer.GetLength());
 			GetMetalDeviceContext().AsyncCopyFromBufferToBuffer(Buffer->CPUBuffer, 0, Buffer->Buffer, 0, Size);
 
-			if (Buffer->GetUsage() & (BUF_Dynamic|BUF_Static))
+			if (Buffer->UsePrivateMemory())
 			{
 				LLM_SCOPE(ELLMTag::VertexBuffer);
 				SafeReleaseMetalBuffer(Buffer->CPUBuffer);
@@ -706,7 +712,7 @@ FVertexBufferRHIRef FMetalDynamicRHI::CreateVertexBuffer_RenderThread(class FRHI
 		}
 		else if (VertexBuffer->Buffer)
 		{
-			if (VertexBuffer->GetUsage() & (BUF_Dynamic|BUF_Static))
+			if (VertexBuffer->UsePrivateMemory())
 			{
 				LLM_SCOPE(ELLMTag::VertexBuffer);
 				SafeReleaseMetalBuffer(VertexBuffer->CPUBuffer);
@@ -794,7 +800,7 @@ void FMetalStagingBuffer::Unlock()
 {
 	check(BackingBuffer);
 	FMetalVertexBuffer* VertexBuffer = ResourceCast(BackingBuffer.GetReference());
-	if (VertexBuffer->CPUBuffer && (VertexBuffer->GetUsage() & (BUF_Dynamic|BUF_Static)))
+	if (VertexBuffer->CPUBuffer && (VertexBuffer->UsePrivateMemory()))
 	{
 		LLM_SCOPE(ELLMTag::VertexBuffer);
 		SafeReleaseMetalBuffer(VertexBuffer->CPUBuffer);
diff --git a/Engine/Source/Runtime/Apple/MetalRHI/Public/MetalResources.h b/Engine/Source/Runtime/Apple/MetalRHI/Public/MetalResources.h
--- a/Engine/Source/Runtime/Apple/MetalRHI/Public/MetalResources.h
+++ b/Engine/Source/Runtime/Apple/MetalRHI/Public/MetalResources.h
@@ -786,6 +786,11 @@ public:
 	 */
 	void Unlock();
 	
+	/**
+	 * Whether to allocate the resource rom private memory or not.
+	 */
+	bool UsePrivateMemory() const;
+	
 	// balsa buffer memory
 	FMetalBuffer Buffer;
 	
@@ -855,6 +860,20 @@ public:
 	/** Resource table containing RHI references. */
 	TArray<TRefCountPtr<FRHIResource> > ResourceTable;
 
+	struct Argument
+	{
+		Argument() {}
+		Argument(FMetalBuffer const& InBuffer, mtlpp::ResourceUsage const InUsage) : Buffer(InBuffer), Usage(InUsage) {}
+		Argument(FMetalTexture const& InTexture, mtlpp::ResourceUsage const InUsage) : Texture(InTexture), Usage(InUsage) {}
+		
+		ns::AutoReleased<FMetalBuffer> Buffer;
+		ns::AutoReleased<FMetalTexture> Texture;
+		mtlpp::ResourceUsage Usage;
+	};
+	
+	TArray<Argument> IndirectArgumentResources;
+	FMetalBuffer IndirectArgumentBuffer;
+	FMetalBuffer IndirectArgumentBufferSideTable;
 };
 
 
diff --git a/Engine/Source/Runtime/Apple/MetalRHI/Public/MetalShaderResources.h b/Engine/Source/Runtime/Apple/MetalRHI/Public/MetalShaderResources.h
--- a/Engine/Source/Runtime/Apple/MetalRHI/Public/MetalShaderResources.h
+++ b/Engine/Source/Runtime/Apple/MetalRHI/Public/MetalShaderResources.h
@@ -133,11 +133,13 @@ struct FMetalShaderBindings
 	TArray<CrossCompiler::FPackedArrayInfo>			PackedGlobalArrays;
 	FMetalShaderResourceTable				ShaderResourceTable;
 	TArray<uint8> 							TypedBufferFormats;
+	TMap<uint8, TArray<uint8>>				ArgumentBufferMasks;
 
     uint32  LinearBuffer;
 	uint32	TypedBuffers;
 	uint32 	InvariantBuffers;
 	uint32 	ConstantBuffers;
+	uint32  ArgumentBuffers;
 	uint16	InOutMask;
 	uint8	NumSamplers;
 	uint8	NumUniformBuffers;
@@ -149,6 +151,7 @@ struct FMetalShaderBindings
         TypedBuffers(0),
 		InvariantBuffers(0),
 		ConstantBuffers(0),
+		ArgumentBuffers(0),
 		InOutMask(0),
 		NumSamplers(0),
 		NumUniformBuffers(0),
@@ -164,10 +167,12 @@ inline FArchive& operator<<(FArchive& Ar, FMetalShaderBindings& Bindings)
 	Ar << Bindings.PackedGlobalArrays;
 	Ar << Bindings.ShaderResourceTable;
 	Ar << Bindings.TypedBufferFormats;
+	Ar << Bindings.ArgumentBufferMasks;
     Ar << Bindings.LinearBuffer;
     Ar << Bindings.TypedBuffers;
 	Ar << Bindings.InvariantBuffers;
 	Ar << Bindings.ConstantBuffers;
+	Ar << Bindings.ArgumentBuffers;
 	Ar << Bindings.InOutMask;
 	Ar << Bindings.NumSamplers;
 	Ar << Bindings.NumUniformBuffers;
