Fix bug that was causing automation tests to fail to compile shaders - tessellation shaders emit a __DSPatch array variable.

diff --git a/Engine/Source/Developer/Apple/MetalShaderFormat/Private/MetalUtils.cpp b/Engine/Source/Developer/Apple/MetalShaderFormat/Private/MetalUtils.cpp
--- a/Engine/Source/Developer/Apple/MetalShaderFormat/Private/MetalUtils.cpp
+++ b/Engine/Source/Developer/Apple/MetalShaderFormat/Private/MetalUtils.cpp
@@ -1249,7 +1249,7 @@ void FMetalCodeBackend::MovePackedUniformsToMain(exec_list* ir, _mesa_glsl_parse
 					check(false);
 					break;
 			}
-			if (Var->type->is_sampler() || Var->type->is_image() || Var->type->is_record() || Var->mode == ir_var_uniform)
+			if (Var->type->is_sampler() || Var->type->is_image() || Var->type->is_array() || Var->type->is_record() || Var->mode == ir_var_uniform)
 			{
 				if (bIsBuffer)
 				{
