Remove potentially unused variable from ue4_stdlib and emit it when required.

diff --git a/Engine/Source/Developer/Apple/MetalShaderFormat/Private/MetalBackend.cpp b/Engine/Source/Developer/Apple/MetalShaderFormat/Private/MetalBackend.cpp
--- a/Engine/Source/Developer/Apple/MetalShaderFormat/Private/MetalBackend.cpp
+++ b/Engine/Source/Developer/Apple/MetalShaderFormat/Private/MetalBackend.cpp
@@ -587,6 +587,7 @@ protected:
 	bool bInsertSideTable;
 
 	bool bRequiresWave;
+	bool bInsertUnknownFormat;
 	bool bNeedsDeviceIndex;
 	
     const char *shaderPrefix()
@@ -2217,6 +2218,7 @@ protected:
 					else
 					{
 						ralloc_asprintf_append(buffer, ", GMetalUnknownFormat");
+						bInsertUnknownFormat = true;
 					}
 					if (bSideTable)
 					{
@@ -2631,6 +2633,7 @@ protected:
 						else
 						{
 							ralloc_asprintf_append(buffer, ", GMetalUnknownFormat");
+							bInsertUnknownFormat = true;
 						}
 						if (bSideTable)
 						{
@@ -2828,6 +2831,7 @@ protected:
 						else
 						{
 							ralloc_asprintf_append(buffer, ", GMetalUnknownFormat");
+							bInsertUnknownFormat = true;
 						}
 						if (bSideTable)
 						{
@@ -4415,6 +4419,7 @@ public:
 		, bImplicitEarlyFragTests(InBackend.Version >= 2)
 		, bInsertSideTable(false)
 		, bRequiresWave(false)
+		, bInsertUnknownFormat(false)
 		, bNeedsDeviceIndex(false)
 	{
 		printable_names = hash_table_ctor(32, hash_table_pointer_hash, hash_table_pointer_compare);
@@ -4491,6 +4496,11 @@ public:
 		{
 			ralloc_asprintf_append(buffer, "\n#define FUNC_ATTRIBS \n\n");
 		}
+		
+		if (bInsertUnknownFormat)
+		{
+			ralloc_asprintf_append(buffer, "\nconstant uint GMetalUnknownFormat = ue4::Unknown;\n\n");
+		}
 
 		// These should work in fragment shaders but Apple are behind the curve on SM6.
 		if (bRequiresWave && Frequency == compute_shader && Backend.Version >= 3)
