Fix Metal DF rendering bug caused by not updating the atomic/invariant-buffer dereference code as part of CL #4553682.

diff --git a/Engine/Shaders/Public/Platform/Metal/MetalCommon.ush b/Engine/Shaders/Public/Platform/Metal/MetalCommon.ush
--- a/Engine/Shaders/Public/Platform/Metal/MetalCommon.ush
+++ b/Engine/Shaders/Public/Platform/Metal/MetalCommon.ush
@@ -5,7 +5,7 @@
 =============================================================================*/
 
 // Update this GUID to improve shader recompilation for Metal only shaders
-// GUID = 7625307457499D7614A9DD98DD45F786
+// GUID = 7655A69A-B6CA-43D6-A68E-E71E0F0ED739
   
 #pragma once
 
diff --git a/Engine/Source/Developer/Apple/MetalShaderFormat/Private/MetalBackend.cpp b/Engine/Source/Developer/Apple/MetalShaderFormat/Private/MetalBackend.cpp
--- a/Engine/Source/Developer/Apple/MetalShaderFormat/Private/MetalBackend.cpp
+++ b/Engine/Source/Developer/Apple/MetalShaderFormat/Private/MetalBackend.cpp
@@ -2251,7 +2251,7 @@ protected:
 						ralloc_asprintf_append(buffer, "min(");
 						tex->coordinate->accept(this);
 						ralloc_asprintf_append(buffer, ",");
-						ralloc_asprintf_append(buffer, "((%s[%d] / sizeof(", BufferSizesName, Index);
+						ralloc_asprintf_append(buffer, "((%s[%d * 2] / sizeof(", BufferSizesName, Index);
 						print_type_pre(Texture->type->inner_type);
 						ralloc_asprintf_append(buffer, "))))]");
 					}
@@ -2665,7 +2665,7 @@ protected:
                             ralloc_asprintf_append(buffer, "min(");
                             deref->image_index->accept(this);
                             ralloc_asprintf_append(buffer, ",");
-							ralloc_asprintf_append(buffer, "((%s[%d] / sizeof(", BufferSizesName, Index);
+							ralloc_asprintf_append(buffer, "((%s[%d * 2] / sizeof(", BufferSizesName, Index);
                             print_type_pre(Texture->type->inner_type);
                             ralloc_asprintf_append(buffer, "))))]");
                         }
