Rework the way atomic operations are handled in MetalBackend and ue4_stdlib to make it possible to perform atomic operations on variables within IABs. Needs a test case to finally sort out.

diff --git a/Engine/Source/Developer/Apple/MetalShaderFormat/Private/MetalBackend.cpp b/Engine/Source/Developer/Apple/MetalShaderFormat/Private/MetalBackend.cpp
--- a/Engine/Source/Developer/Apple/MetalShaderFormat/Private/MetalBackend.cpp
+++ b/Engine/Source/Developer/Apple/MetalShaderFormat/Private/MetalBackend.cpp
@@ -1045,9 +1045,10 @@ protected:
 											   );
 						if (bIsAtomic)
 						{
-							ralloc_asprintf_append(buffer, "atomic_");
+							ralloc_asprintf_append(buffer, "typed_buffer<");
 							check(BufferIndex < 8);
 							print_type_pre(PtrType->inner_type);
+							ralloc_asprintf_append(buffer, ">");
 						}
 						else
 						{
@@ -3473,20 +3474,6 @@ protected:
 
 	virtual void visit(ir_atomic *ir) override
 	{
-		const char* SharedAtomicFunctions[ir_atomic_count] = 
-		{
-			"atomic_fetch_add_explicit",
-			"atomic_fetch_and_explicit",
-			"atomic_fetch_min_explicit",
-			"atomic_fetch_max_explicit",
-			"atomic_fetch_or_explicit",
-			"atomic_fetch_xor_explicit",
-			"atomic_exchange_explicit",
-			"atomic_compare_exchange_weak_explicit",
-			"atomic_load_explicit",
-			"atomic_store_explicit",
-		};
-		static_assert(sizeof(SharedAtomicFunctions) / sizeof(SharedAtomicFunctions[0]) == ir_atomic_count, "Mismatched entries!");
 /*
 		const char *imageAtomicFunctions[] =
 		{
@@ -3508,10 +3495,57 @@ protected:
 			ir->lhs->accept(this);
 			ralloc_asprintf_append(buffer, " = ");
 		}
-		//if (!is_image)
+		if (is_image)
 		{
-			ralloc_asprintf_append(buffer, "%s(&", SharedAtomicFunctions[ir->operation]);
-			ir->memory_ref->accept(this);
+			const char* SharedAtomicFunctions[ir_atomic_count] =
+			{
+				"fetch_add_atomic",
+				"fetch_and_atomic",
+				"fetch_min_atomic",
+				"fetch_max_atomic",
+				"fetch_or_atomic",
+				"fetch_xor_atomic",
+				"exchange_atomic",
+				"compare_exchange_weak_atomic",
+				"load_atomic",
+				"store_atomic",
+			};
+			static_assert(sizeof(SharedAtomicFunctions) / sizeof(SharedAtomicFunctions[0]) == ir_atomic_count, "Mismatched entries!");
+
+			ir_dereference_image* atomic = ir->memory_ref->as_dereference_image();
+
+			int BufferIndex = 0;
+			char const* BufferSizesName = "BufferSizes";
+
+			ir_variable* image_var = atomic->image->variable_referenced();
+			if (image_var->mode == ir_var_temporary)
+			{
+				// IAB sampling path
+				ir_variable* IABVariable = Backend.IABVariablesMap.FindChecked(image_var);
+				int FieldIndex = IABVariable->type->field_index(image_var->name);
+				for (int i = 0; i < FieldIndex; i++)
+				{
+					if (IABVariable->type->fields.structure[i].type->sampler_buffer)
+					{
+						BufferIndex++;
+					}
+				}
+
+				BufferSizesName = ralloc_asprintf(ParseState, "%s.BufferSizes", IABVariable->name);
+			}
+			else
+			{
+				// Function argument path
+				BufferIndex = Buffers.GetIndex(image_var);
+			}
+			check(BufferIndex >= 0 && BufferIndex <= 30);
+
+			ralloc_asprintf_append(buffer, "buffer_atomic<memory_order_relaxed>::%s<", SharedAtomicFunctions[ir->operation]);
+			print_type_pre(ir->memory_ref->type);
+			ralloc_asprintf_append(buffer, ", %d>(", BufferIndex);
+			atomic->image->accept(this);
+			ralloc_asprintf_append(buffer, ", %s, ", BufferSizesName);
+			atomic->image_index->accept(this);
 			if (ir->operands[0])
 			{
 				ralloc_asprintf_append(buffer, ", ");
@@ -3522,26 +3556,39 @@ protected:
 				ralloc_asprintf_append(buffer, ", ");
 				ir->operands[1]->accept(this);
 			}
-			ralloc_asprintf_append(buffer, ", memory_order_relaxed)");
+			ralloc_asprintf_append(buffer, ")");
 		}
-/*
 		else
 		{
-			ir_dereference_image *image = ir->memory_ref->as_dereference_image();
-			ralloc_asprintf_append(buffer, " = %s(",
-				imageAtomicFunctions[ir->operation]);
-			image->image->accept(this);
-			ralloc_asprintf_append(buffer, ", ");
-			image->image_index->accept(this);
-			ralloc_asprintf_append(buffer, ", ");
-			ir->operands[0]->accept(this);
+			const char* SharedAtomicFunctions[ir_atomic_count] =
+			{
+				"atomic_fetch_add_explicit",
+				"atomic_fetch_and_explicit",
+				"atomic_fetch_min_explicit",
+				"atomic_fetch_max_explicit",
+				"atomic_fetch_or_explicit",
+				"atomic_fetch_xor_explicit",
+				"atomic_exchange_explicit",
+				"atomic_compare_exchange_weak_explicit",
+				"atomic_load_explicit",
+				"atomic_store_explicit",
+			};
+			static_assert(sizeof(SharedAtomicFunctions) / sizeof(SharedAtomicFunctions[0]) == ir_atomic_count, "Mismatched entries!");
+			
+			ralloc_asprintf_append(buffer, "%s(&", SharedAtomicFunctions[ir->operation]);
+			ir->memory_ref->accept(this);
+			if (ir->operands[0])
+			{
+				ralloc_asprintf_append(buffer, ", ");
+				ir->operands[0]->accept(this);
+			}
 			if (ir->operands[1])
 			{
 				ralloc_asprintf_append(buffer, ", ");
 				ir->operands[1]->accept(this);
 			}
 			ralloc_asprintf_append(buffer, ", memory_order_relaxed)");
-		}*/
+		}
 	}
 
 	/**
