Emit Metal tessellation shader stage-in variables properly after the Metal shader compiler changes.

diff --git a/Engine/Shaders/Public/Platform/Metal/MetalCommon.ush b/Engine/Shaders/Public/Platform/Metal/MetalCommon.ush
--- a/Engine/Shaders/Public/Platform/Metal/MetalCommon.ush
+++ b/Engine/Shaders/Public/Platform/Metal/MetalCommon.ush
@@ -5,7 +5,7 @@
 =============================================================================*/
 
 // Update this GUID to improve shader recompilation for Metal only shaders
-// GUID = 7655A69A-B6CA-43D6-A68E-E71E0F0ED739
+// GUID = 0FDEF9A8-0978-44ED-B2ED-046924F18E72
   
 #pragma once
 
diff --git a/Engine/Source/Developer/Apple/MetalShaderFormat/Private/MetalBackend.cpp b/Engine/Source/Developer/Apple/MetalShaderFormat/Private/MetalBackend.cpp
--- a/Engine/Source/Developer/Apple/MetalShaderFormat/Private/MetalBackend.cpp
+++ b/Engine/Source/Developer/Apple/MetalShaderFormat/Private/MetalBackend.cpp
@@ -1275,7 +1275,7 @@ protected:
 							var->semantic
 							);
 					}
-					else if(var->semantic && var->type->is_record())
+					else if(var->semantic && strcmp(var->semantic, "stage_in") && var->type->is_record())
 					{
 						ralloc_asprintf_append(buffer,"device ");
 						print_type_pre(var->type);
