Fix a missing new-line causing IABs to not be properly bound because they were not being parsed into the binding table.

diff --git a/Engine/Source/Developer/Apple/MetalShaderFormat/Private/MetalBackend.cpp b/Engine/Source/Developer/Apple/MetalShaderFormat/Private/MetalBackend.cpp
--- a/Engine/Source/Developer/Apple/MetalShaderFormat/Private/MetalBackend.cpp
+++ b/Engine/Source/Developer/Apple/MetalShaderFormat/Private/MetalBackend.cpp
@@ -4262,7 +4262,7 @@ protected:
 				{
 					check(foundSideTable == false);
 					foundSideTable = true;
-					ralloc_asprintf_append(buffer, "// @SideTable: %s(%d)", Var->name, i);
+					ralloc_asprintf_append(buffer, "// @SideTable: %s(%d)\n", Var->name, i);
 				}
 			}
 		}
