Tie Metal PCH compilation to the specific Xcode build, not just the metal compiler version.

diff --git a/Engine/Source/Developer/Apple/MetalShaderFormat/Private/MetalDerivedData.cpp b/Engine/Source/Developer/Apple/MetalShaderFormat/Private/MetalDerivedData.cpp
--- a/Engine/Source/Developer/Apple/MetalShaderFormat/Private/MetalDerivedData.cpp
+++ b/Engine/Source/Developer/Apple/MetalShaderFormat/Private/MetalDerivedData.cpp
@@ -140,17 +140,22 @@ const TCHAR* FMetalShaderBytecodeCooker::GetVersionString() const
 
 FString FMetalShaderBytecodeCooker::GetPluginSpecificCacheKeySuffix() const
 {
-	FString CompilerVersion = GetMetalCompilerVersion(MetalShaderFormatToLegacyShaderPlatform(Job.ShaderFormat));
+	FString CompilerVersion = Job.CompilerVersion;
 	FString CompilerPath = GetMetalToolsPath(MetalShaderFormatToLegacyShaderPlatform(Job.ShaderFormat));
 
+	uint64 BuildVersion = 0;
+	uint32 XcodeVersion = GetXcodeVersion(BuildVersion);
+	
 	// PCHs need the modifiction time (in secs. since UTC Epoch) to ensure that the result can be used with the current version of the file
 	uint64 ModTime = 0;
 	if (Job.bCompileAsPCH)
 	{
 		ModificationTimeRemoteFile(*Job.InputFile, ModTime);
+		
+		CompilerVersion += FString::Printf(TEXT("xc%u%llu"), XcodeVersion, BuildVersion);
 	}
 
-	FString VersionedName = FString::Printf(TEXT("%s%u%u%llu%s%s%s%s%s%s%s%d%d"), *Job.ShaderFormat.GetPlainNameString(), Job.SourceCRCLen, Job.SourceCRC, ModTime, *Job.Hash.ToString(), *Job.CompilerVersion, *Job.MinOSVersion, *Job.DebugInfo, *Job.MathMode, *Job.Standard, Job.bRetainObjectFile ? TEXT("+O") : TEXT(""), GetTypeHash(CompilerPath), GetTypeHash(Job.Defines));
+	FString VersionedName = FString::Printf(TEXT("%s%u%u%llu%s%s%s%s%s%s%s%d%d"), *Job.ShaderFormat.GetPlainNameString(), Job.SourceCRCLen, Job.SourceCRC, ModTime, *Job.Hash.ToString(), *CompilerVersion, *Job.MinOSVersion, *Job.DebugInfo, *Job.MathMode, *Job.Standard, Job.bRetainObjectFile ? TEXT("+O") : TEXT(""), GetTypeHash(CompilerPath), GetTypeHash(Job.Defines));
 	// get rid of some not so filename-friendly characters ('=',' ' -> '_')
 	VersionedName = VersionedName.Replace(TEXT("="), TEXT("_")).Replace(TEXT(" "), TEXT("_"));
 
@@ -206,11 +211,7 @@ bool FMetalShaderBytecodeCooker::Build(TArray<uint8>& OutData)
 	{
 		CopyLocalFileToRemote(Job.InputFile, RemoteInputFile);
 
-		uint64 XcodeBuildVers = 0;
-		uint16 XcodeVers = GetXcodeVersion(XcodeBuildVers);
-		uint16 XcodeMajorVers = ((XcodeVers >> 8) & 0xff);
-		
-		// PCH 
+		// PCH
 		bool bUseSharedPCH = Job.InputPCHFile.Len() && IFileManager::Get().FileExists(*Job.InputPCHFile);
 		if (bUseSharedPCH)
         {
