Removing instanced-stereo code by having "ResolvedView" become "View" doesn't work anymore for Metal because the structure splitting ends up redefining the variable. It is easier just to have everything be the same and let the compiler strip away the unnecessary elements.

diff --git a/Engine/Source/Developer/Apple/MetalShaderFormat/Private/MetalShaderCompiler.cpp b/Engine/Source/Developer/Apple/MetalShaderFormat/Private/MetalShaderCompiler.cpp
--- a/Engine/Source/Developer/Apple/MetalShaderFormat/Private/MetalShaderCompiler.cpp
+++ b/Engine/Source/Developer/Apple/MetalShaderFormat/Private/MetalShaderCompiler.cpp
@@ -2221,12 +2221,6 @@ void CompileShader_Metal(const FShaderCompilerInput& _Input,FShaderCompilerOutpu
 		}
 	}
 
-	if (Input.ShaderFormat != NAME_SF_METAL_SM5)
-	{
-		// Disable instanced stereo on everything but Metal SM5 for 10.13+
-		StripInstancedStereo(PreprocessedShader);
-	}
-
 	char* MetalShaderSource = NULL;
 	char* ErrorLog = NULL;
 
