Fix creation of texture 2d arrays and handling of texture2darray.gather in Metal.

diff --git a/Engine/Source/Developer/Apple/MetalShaderFormat/Private/MetalBackend.cpp b/Engine/Source/Developer/Apple/MetalShaderFormat/Private/MetalBackend.cpp
--- a/Engine/Source/Developer/Apple/MetalShaderFormat/Private/MetalBackend.cpp
+++ b/Engine/Source/Developer/Apple/MetalShaderFormat/Private/MetalBackend.cpp
@@ -2322,7 +2322,56 @@ protected:
 			}
 			
 			// Coord
-			tex->coordinate->accept(this);
+			if (tex->sampler->type->sampler_array)
+			{
+				// Need to split the coordinate
+				char const* CoordSwizzle = "";
+				char const* IndexSwizzle = "y";
+				switch(tex->sampler->type->sampler_dimensionality)
+				{
+					case GLSL_SAMPLER_DIM_1D:
+					{
+						break;
+					}
+					case GLSL_SAMPLER_DIM_2D:
+					case GLSL_SAMPLER_DIM_RECT:
+					{
+						CoordSwizzle = "y";
+						IndexSwizzle = "z";
+						break;
+					}
+					case GLSL_SAMPLER_DIM_3D:
+					{
+						CoordSwizzle = "yz";
+						IndexSwizzle = "w";
+						break;
+					}
+					case GLSL_SAMPLER_DIM_CUBE:
+					{
+						CoordSwizzle = "yz";
+						IndexSwizzle = "w";
+						break;
+					}
+					case GLSL_SAMPLER_DIM_BUF:
+					case GLSL_SAMPLER_DIM_EXTERNAL:
+					default:
+					{
+						check(0);
+						break;
+					}
+				}
+				
+				ralloc_asprintf_append(buffer, "(");
+				tex->coordinate->accept(this);
+				
+				ralloc_asprintf_append(buffer, ").x%s, (uint)(", CoordSwizzle);
+				tex->coordinate->accept(this);
+				ralloc_asprintf_append(buffer, ").%s", IndexSwizzle);
+			}
+			else
+			{
+				tex->coordinate->accept(this);
+			}
 
 			if (tex->shadow_comparitor)
 			{
diff --git a/Engine/Source/Runtime/Apple/MetalRHI/Private/MetalTexture.cpp b/Engine/Source/Runtime/Apple/MetalRHI/Private/MetalTexture.cpp
--- a/Engine/Source/Runtime/Apple/MetalRHI/Private/MetalTexture.cpp
+++ b/Engine/Source/Runtime/Apple/MetalRHI/Private/MetalTexture.cpp
@@ -841,6 +841,10 @@ FMetalSurface::FMetalSurface(ERHIResourceType ResourceType, EPixelFormat Format,
 				Desc.SetArrayLength(ArraySize * 6);
 			}
 		}
+		else
+		{
+			Desc.SetTextureType(mtlpp::TextureType::Texture2DArray);
+		}
 	}
 	Desc.SetMipmapLevelCount(NumMips);
 	
