Fix the swizzle on Metal for reads from a RWTexture as the Metal function always returns as a vector-4.

diff --git a/Engine/Source/Developer/Apple/MetalShaderFormat/Private/MetalBackend.cpp b/Engine/Source/Developer/Apple/MetalShaderFormat/Private/MetalBackend.cpp
--- a/Engine/Source/Developer/Apple/MetalShaderFormat/Private/MetalBackend.cpp
+++ b/Engine/Source/Developer/Apple/MetalShaderFormat/Private/MetalBackend.cpp
@@ -2626,10 +2626,27 @@ protected:
 			{
 				if (bIsRWTexture)
 				{
+					ralloc_asprintf_append(buffer, "(");
 					deref->image->accept(this);
 					ralloc_asprintf_append(buffer, ".read(");
 					deref->image_index->accept(this);
 					ralloc_asprintf_append(buffer, ")");
+					switch(dst_elements)
+					{
+						case 1:
+							ralloc_asprintf_append(buffer, ".x");
+							break;
+						case 2:
+							ralloc_asprintf_append(buffer, ".xy");
+							break;
+						case 3:
+							ralloc_asprintf_append(buffer, ".xyz");
+							break;
+						case 4:
+						default:
+							break;
+					}
+					ralloc_asprintf_append(buffer, ")");
 				}
 				else
 				{
