Don't zero-init threadgroup variables in Metal shaders anymore - it doesn't seem to be necessary these days.

diff --git a/Engine/Source/Developer/Apple/MetalShaderFormat/Private/MetalBackend.cpp b/Engine/Source/Developer/Apple/MetalShaderFormat/Private/MetalBackend.cpp
--- a/Engine/Source/Developer/Apple/MetalShaderFormat/Private/MetalBackend.cpp
+++ b/Engine/Source/Developer/Apple/MetalShaderFormat/Private/MetalBackend.cpp
@@ -1404,7 +1404,7 @@ protected:
 				var->constant_value->accept(this);
 			}
 		}
-		else if ((Backend.bZeroInitialise) && (var->type->base_type != GLSL_TYPE_STRUCT) && (var->mode == ir_var_auto || var->mode == ir_var_temporary || var->mode == ir_var_shared) && (Buffers.AtomicVariables.find(var) == Buffers.AtomicVariables.end()))
+		else if ((Backend.bZeroInitialise) && (var->mode != ir_var_shared) && (var->type->base_type != GLSL_TYPE_STRUCT) && (var->mode == ir_var_auto || var->mode == ir_var_temporary) && (Buffers.AtomicVariables.find(var) == Buffers.AtomicVariables.end()))
 		{
 			// @todo UE-34355 temporary workaround for 10.12 shader compiler error - really all arrays should be zero'd but only threadgroup shared initialisation works on the Beta drivers.
 			if (!is_struct_type(var->type) && (var->type->base_type != GLSL_TYPE_ARRAY || var->mode == ir_var_shared) && (var->type->is_numeric() || (var->type->base_type == GLSL_TYPE_ARRAY)))
