The buffer_argument type should only be emitted if Metal version is 2.0 or later.

diff --git a/Engine/Source/Developer/Apple/MetalShaderFormat/Private/MetalBackend.cpp b/Engine/Source/Developer/Apple/MetalShaderFormat/Private/MetalBackend.cpp
--- a/Engine/Source/Developer/Apple/MetalShaderFormat/Private/MetalBackend.cpp
+++ b/Engine/Source/Developer/Apple/MetalShaderFormat/Private/MetalBackend.cpp
@@ -704,20 +704,27 @@ protected:
 		{
 			if (t->sampler_buffer)
 			{
-				if (!strncmp(t->name, "RWBuffer<", 9))
+				if (Backend.Version > 2)
 				{
-					ralloc_asprintf_append(buffer, "buffer_argument<");
-					print_type_pre(t->inner_type);
-					ralloc_asprintf_append(buffer, ", access::read_write>");
+					if (!strncmp(t->name, "RWBuffer<", 9))
+					{
+						ralloc_asprintf_append(buffer, "buffer_argument<");
+						print_type_pre(t->inner_type);
+						ralloc_asprintf_append(buffer, ", access::read_write>");
+					}
+					else
+					{
+						if (strncmp(t->HlslName, "RW", 2))
+						{
+							ralloc_asprintf_append(buffer, "const ");
+						}
+						print_type_pre(t->inner_type);
+						ralloc_asprintf_append(buffer, "*");
+					}
 				}
 				else
 				{
-					if (strncmp(t->HlslName, "RW", 2))
-					{
-						ralloc_asprintf_append(buffer, "const ");
-					}
 					print_type_pre(t->inner_type);
-					ralloc_asprintf_append(buffer, "*");
 				}
 			}
 			else
@@ -808,9 +815,16 @@ protected:
 					case GLSL_SAMPLER_DIM_BUF:
 						// Typed buffer read
 						check(t->inner_type);
-						ralloc_asprintf_append(buffer, "buffer_argument<");
-						print_base_type(t->inner_type);
-						ralloc_asprintf_append(buffer, ">");
+						if (Backend.Version > 2)
+						{
+							ralloc_asprintf_append(buffer, "buffer_argument<");
+							print_base_type(t->inner_type);
+							ralloc_asprintf_append(buffer, ">");
+						}
+						else
+						{
+							print_base_type(t->inner_type);
+						}
 						break;
 					case GLSL_SAMPLER_DIM_RECT:
 					case GLSL_SAMPLER_DIM_EXTERNAL:
