Emulate texture2d-atomics for Metal by creating textures marked as TexCreate_UAV|TexCreate_NoTiling from MTLBuffers as Linear Textures, which means we can bind the buffer for the atomic operations and the texture for normal operations. There may be a performance penalty and unexpected cache interactions with this approach and it is not used anywhere in the engine yet.

diff --git a/Engine/Source/Developer/Apple/MetalShaderFormat/Private/MetalBackend.cpp b/Engine/Source/Developer/Apple/MetalShaderFormat/Private/MetalBackend.cpp
--- a/Engine/Source/Developer/Apple/MetalShaderFormat/Private/MetalBackend.cpp
+++ b/Engine/Source/Developer/Apple/MetalShaderFormat/Private/MetalBackend.cpp
@@ -2954,9 +2954,32 @@ protected:
 		}
 		else if ( deref->op == ir_image_dimensions)
 		{
-			ralloc_asprintf_append( buffer, "imageSize( " );
+			// Convert from:
+			//	HLSL
+			//		int w, h;
+			//		T.GetDimensions({lod, }w, h);
+			// GLSL
+			//		ivec2 Temp;
+			//		Temp = textureSize(T{, lod});
+			// Metal
+			//		int2 Temp = int2((int)T.get_width({lod}), (int)T.get_height({lod}));
+			ralloc_asprintf_append(buffer, "int2(");
 			deref->image->accept(this);
-			ralloc_asprintf_append(buffer, ")");
+			ralloc_asprintf_append(buffer, ".get_width(");
+			
+			if (deref->image_index)
+			{
+				deref->image_index->accept(this);
+			}
+			ralloc_asprintf_append(buffer, "), (int)");
+			
+			deref->image->accept(this);
+			ralloc_asprintf_append(buffer, ".get_height(");
+			if (deref->image_index)
+			{
+				deref->image_index->accept(this);
+			}
+			ralloc_asprintf_append(buffer, "))");
 		}
 		else
 		{
@@ -3568,7 +3591,7 @@ protected:
 		};
 */
 		check(scope_depth > 0);
-		const bool is_image = ir->memory_ref->as_dereference_image() != NULL;
+		const bool is_image = ir->memory_ref->as_dereference_image() != NULL || ir->memory_ref->type->is_image();
 
 		if (ir->lhs)
 		{
@@ -3591,13 +3614,31 @@ protected:
 				"store_atomic",
 			};
 			static_assert(sizeof(SharedAtomicFunctions) / sizeof(SharedAtomicFunctions[0]) == ir_atomic_count, "Mismatched entries!");
-
-			ir_dereference_image* atomic = ir->memory_ref->as_dereference_image();
-
 			int BufferIndex = 0;
 			char const* BufferSizesName = "BufferSizes";
-
-			ir_variable* image_var = atomic->image->variable_referenced();
+			
+			ir_dereference_image* atomic = ir->memory_ref->as_dereference_image();
+			ir_dereference_variable* deref = ir->memory_ref->as_dereference_variable();
+			ir_variable* image_var = nullptr;
+			
+			ir_rvalue *image = nullptr;
+			ir_rvalue *image_index = nullptr;
+			ir_rvalue *operands[2] = {nullptr, nullptr};
+			if (atomic)
+			{
+				image_var = atomic->image->variable_referenced();
+				image = atomic->image;
+				image_index = atomic->image_index;
+				operands[0] = ir->operands[0];
+				operands[1] = ir->operands[1];
+			}
+			else
+			{
+				check(deref);
+				image_var = deref->variable_referenced();
+				image = deref;
+				image_index = ir->operands[0];
+			}
 			if (image_var->mode == ir_var_temporary)
 			{
 				// IAB sampling path
@@ -3621,20 +3662,20 @@ protected:
 			check(BufferIndex >= 0 && BufferIndex <= 30);
 
 			ralloc_asprintf_append(buffer, "buffer_atomic<memory_order_relaxed>::%s<", SharedAtomicFunctions[ir->operation]);
-			print_type_pre(ir->memory_ref->type);
+			print_type_pre(image_var->type->inner_type);
 			ralloc_asprintf_append(buffer, ", %d>(", BufferIndex);
-			atomic->image->accept(this);
+			image->accept(this);
 			ralloc_asprintf_append(buffer, ", %s, ", BufferSizesName);
-			atomic->image_index->accept(this);
-			if (ir->operands[0])
+			image_index->accept(this);
+			if (operands[0])
 			{
 				ralloc_asprintf_append(buffer, ", ");
-				ir->operands[0]->accept(this);
+				operands[0]->accept(this);
 			}
-			if (ir->operands[1])
+			if (operands[1])
 			{
 				ralloc_asprintf_append(buffer, ", ");
-				ir->operands[1]->accept(this);
+				operands[1]->accept(this);
 			}
 			ralloc_asprintf_append(buffer, ")");
 		}
@@ -4774,6 +4815,103 @@ public:
 	}
 };
 
+struct FMetalAtomicTexture2DVisitor : public ir_hierarchical_visitor
+{
+	exec_list* Instructions;
+	_mesa_glsl_parse_state* ParseState;
+	
+	FMetalAtomicTexture2DVisitor(exec_list* ir, _mesa_glsl_parse_state* state)
+	: Instructions(ir)
+	, ParseState(state)
+	{
+		
+	}
+	
+	~FMetalAtomicTexture2DVisitor()
+	{
+		
+	}
+	
+	virtual ir_visitor_status visit_leave(class ir_atomic * ir) override final
+	{
+		const bool is_image = ir->memory_ref->as_dereference_image() != NULL;
+		if (is_image)
+		{
+			ir_dereference_image* atomic = ir->memory_ref->as_dereference_image();
+			ir_variable* image_var = atomic->image->variable_referenced();
+			switch (image_var->type->sampler_dimensionality)
+			{
+				case GLSL_SAMPLER_DIM_BUF:
+					break;
+				case GLSL_SAMPLER_DIM_2D:
+				{
+					// Not handling IABs yet
+					check(image_var->mode == ir_var_uniform);
+					
+					char const* newName = ralloc_asprintf(ParseState, "%s_atomic", image_var->name);
+					ir_variable* newVar = ParseState->symbols->get_variable(newName);
+					if (!newVar)
+					{
+						glsl_type const* BufferType = glsl_type::GetStructuredBufferInstance("RWStructuredBuffer", image_var->type->inner_type);
+						newVar = new(ParseState)ir_variable(BufferType, newName, ir_var_uniform);
+						newVar->used = 1;
+						
+						image_var->constant_value = (ir_constant*)newVar;
+						
+						Instructions->push_head(newVar);
+						
+						ParseState->symbols->add_variable(newVar);
+					}
+					check(newVar);
+					
+					ir_dereference_variable* DerefVar = new(ParseState)ir_dereference_variable(newVar);
+					
+					const glsl_type *res_type = glsl_type::get_instance(GLSL_TYPE_INT, 2, 1);
+					ir_variable* temp = new(ParseState)ir_variable(res_type, nullptr, ir_var_temporary);
+					
+					ir_dereference_image* DerefOld = new(ParseState)ir_dereference_image(atomic->image->clone(ParseState, nullptr), new(ParseState) ir_constant(0.0f), ir_image_dimensions);
+					DerefOld->type = res_type;
+					
+					ir_assignment* Assign = new(ParseState)ir_assignment(new(ParseState)ir_dereference_variable(temp), DerefOld);
+					
+					ir->insert_before(temp);
+					ir->insert_before(Assign);
+					
+					unsigned int XSwizzle[] = { 0 };
+					ir_swizzle* Width = new(ParseState)ir_swizzle(new(ParseState)ir_dereference_variable(temp), XSwizzle, 1);
+					
+					ir_swizzle* XCoord = new(ParseState)ir_swizzle(atomic->image_index->clone(ParseState, nullptr), XSwizzle, 1);
+					
+					unsigned int YSwizzle[] = { 1 };
+					ir_swizzle* YCoord = new(ParseState)ir_swizzle(atomic->image_index->clone(ParseState, nullptr), YSwizzle, 1);
+					
+					ir_expression* Mul = new (ParseState)ir_expression(ir_binop_mul, Width, XCoord);
+					ir_expression* Add = new (ParseState)ir_expression(ir_binop_add, Mul, YCoord);
+					
+					ir_dereference_image* DerefImage = new(ParseState)ir_dereference_image(DerefVar, Add, ir_image_access);
+					
+					ir->memory_ref = DerefImage;
+					break;
+				}
+				default:
+					if (!image_var->type->sampler_buffer)
+					{
+						_mesa_glsl_error(ParseState, "Metal doesn't allow atommic operations on RWTexture %s", image_var->name);
+					}
+					break;
+			}
+		}
+		
+		return ir_hierarchical_visitor::visit_leave(ir);
+	}
+};
+
+void FMetalCodeBackend::FixupTextureAtomics(exec_list* ir, _mesa_glsl_parse_state* state)
+{
+	FMetalAtomicTexture2DVisitor Visitor(ir, state);
+	Visitor.run(ir);
+}
+
 char* FMetalCodeBackend::GenerateCode(exec_list* ir, _mesa_glsl_parse_state* state, EHlslShaderFrequency Frequency)
 {
 	// We'll need this Buffers info for the [[buffer()]] index
@@ -4798,6 +4936,8 @@ char* FMetalCodeBackend::GenerateCode(exec_list* ir, _mesa_glsl_parse_state* sta
 	// Move all inputs & outputs to structs for Metal
 	PackInputsAndOutputs(ir, state, Frequency, visitor.input_variables);
 	
+	FixupTextureAtomics(ir, state);
+	
 	FindAtomicVariables(ir, Buffers.AtomicVariables);
 	
 	// ir_var_uniform instances be global, so move them as arguments to main
diff --git a/Engine/Source/Developer/Apple/MetalShaderFormat/Private/MetalBackend.h b/Engine/Source/Developer/Apple/MetalShaderFormat/Private/MetalBackend.h
--- a/Engine/Source/Developer/Apple/MetalShaderFormat/Private/MetalBackend.h
+++ b/Engine/Source/Developer/Apple/MetalShaderFormat/Private/MetalBackend.h
@@ -118,6 +118,7 @@ struct FMetalCodeBackend : public FCodeBackend
 	void BreakPrecisionChangesVisitor(exec_list* ir, _mesa_glsl_parse_state* State);
 	void FixupMetalBaseOffsets(exec_list* ir, _mesa_glsl_parse_state* state, EHlslShaderFrequency Frequency);
 	void InsertSamplerStates(exec_list* ir, _mesa_glsl_parse_state* State);
+	void FixupTextureAtomics(exec_list* ir, _mesa_glsl_parse_state* state);
 
 	TMap<ir_variable*, TSet<uint8>> IABVariableMask;
 	TMap<ir_variable*, ir_variable*> IABVariablesMap;
diff --git a/Engine/Source/Developer/Apple/MetalShaderFormat/Private/MetalUtils.h b/Engine/Source/Developer/Apple/MetalShaderFormat/Private/MetalUtils.h
--- a/Engine/Source/Developer/Apple/MetalShaderFormat/Private/MetalUtils.h
+++ b/Engine/Source/Developer/Apple/MetalShaderFormat/Private/MetalUtils.h
@@ -286,7 +286,7 @@ struct FBuffers
         for(uint64 i = 0; i < 64ull && !ITextures.empty(); i++)
         {
             check(i < 64);
-            if (!(UAVIndices & (1ull << i)))
+            if (!(UAVIndices & (1ull << i)) || (AllBuffers[i] == (ir_variable*)ITextures.front()->constant_value))
             {
                 // i *must* be less than 8 as we only support indices 0-7 for UAVs and beyond that we can't bind it.
                 if (i >= 8)
diff --git a/Engine/Source/Runtime/Apple/MetalRHI/Private/MetalStateCache.cpp b/Engine/Source/Runtime/Apple/MetalRHI/Private/MetalStateCache.cpp
--- a/Engine/Source/Runtime/Apple/MetalRHI/Private/MetalStateCache.cpp
+++ b/Engine/Source/Runtime/Apple/MetalRHI/Private/MetalStateCache.cpp
@@ -1673,6 +1673,11 @@ void FMetalStateCache::SetShaderUnorderedAccessView(EShaderFrequency ShaderStage
 				FPlatformAtomics::InterlockedExchange(&Source->Written, 1);
 				
 				SetShaderTexture(ShaderStage, Surface->Texture, BindIndex, mtlpp::ResourceUsage(mtlpp::ResourceUsage::Read | mtlpp::ResourceUsage::Write));
+				if ((Source->Flags & (TexCreate_UAV|TexCreate_NoTiling)) == (TexCreate_UAV|TexCreate_NoTiling))
+				{
+					FMetalBuffer Buffer(Surface->Texture.GetBuffer(), false);
+					SetShaderBuffer(ShaderStage, Buffer, nullptr, 0, Surface->Texture.GetBuffer().GetLength(), BindIndex, mtlpp::ResourceUsage(mtlpp::ResourceUsage::Read | mtlpp::ResourceUsage::Write), PF_MAX);
+				}
 			}
 			else
 			{
diff --git a/Engine/Source/Runtime/Apple/MetalRHI/Private/MetalTexture.cpp b/Engine/Source/Runtime/Apple/MetalRHI/Private/MetalTexture.cpp
--- a/Engine/Source/Runtime/Apple/MetalRHI/Private/MetalTexture.cpp
+++ b/Engine/Source/Runtime/Apple/MetalRHI/Private/MetalTexture.cpp
@@ -863,7 +863,7 @@ FMetalSurface::FMetalSurface(ERHIResourceType ResourceType, EPixelFormat Format,
 			Desc.SetResourceOptions((mtlpp::ResourceOptions)(mtlpp::ResourceOptions::CpuCacheModeDefaultCache|mtlpp::ResourceOptions::StorageModeShared));
 #endif
 		}
-		else if(((Flags & (TexCreate_NoTiling)) && !(Flags & (TexCreate_FastVRAM|TexCreate_DepthStencilTargetable|TexCreate_RenderTargetable))))
+		else if(((Flags & (TexCreate_NoTiling)) && !(Flags & (TexCreate_FastVRAM|TexCreate_DepthStencilTargetable|TexCreate_RenderTargetable|TexCreate_UAV))))
 		{
 #if PLATFORM_MAC
 			Desc.SetCpuCacheMode(mtlpp::CpuCacheMode::WriteCombined);
@@ -964,7 +964,19 @@ FMetalSurface::FMetalSurface(ERHIResourceType ResourceType, EPixelFormat Format,
 	}
 	else
 	{
-		Texture = GetMetalDeviceContext().CreateTexture(this, Desc);
+		if ((Flags & (TexCreate_UAV|TexCreate_NoTiling)) != (TexCreate_UAV|TexCreate_NoTiling))
+		{
+			Texture = GetMetalDeviceContext().CreateTexture(this, Desc);
+		}
+		else
+		{
+			mtlpp::Device Device = GetMetalDeviceContext().GetDevice();
+			mtlpp::SizeAndAlign SizeAlign = Device.HeapTextureSizeAndAlign(Desc);
+			FMetalPooledBufferArgs Args(Device, SizeAlign.Size, mtlpp::StorageMode::Private);
+			FMetalBuffer Buffer = GetMetalDeviceContext().CreatePooledBuffer(Args);
+			Texture = Buffer.NewTexture(Desc, 0, SizeAlign.Size);
+		}
+		
 		if (Texture.GetPtr() == nil)
 		{
 			UE_LOG(LogMetal, Fatal, TEXT("Failed to create texture, desc %s"), *FString([Desc description]));
diff --git a/Engine/Source/Runtime/Apple/MetalRHI/Private/MetalUAV.cpp b/Engine/Source/Runtime/Apple/MetalRHI/Private/MetalUAV.cpp
--- a/Engine/Source/Runtime/Apple/MetalRHI/Private/MetalUAV.cpp
+++ b/Engine/Source/Runtime/Apple/MetalRHI/Private/MetalUAV.cpp
@@ -496,9 +496,10 @@ void FMetalRHICommandContext::RHIClearTinyUAV(FUnorderedAccessViewRHIParamRef Un
 {
 	@autoreleasepool {
 	FMetalUnorderedAccessView* UnorderedAccessView = ResourceCast(UnorderedAccessViewRHI);
-	if (UnorderedAccessView->SourceView->SourceStructuredBuffer || UnorderedAccessView->SourceView->SourceVertexBuffer || UnorderedAccessView->SourceView->SourceIndexBuffer)
+	FMetalSurface* Surface = UnorderedAccessView->SourceView->SourceTexture ? GetMetalSurfaceFromRHITexture(UnorderedAccessView->SourceView->SourceTexture) : nullptr;
+	if (UnorderedAccessView->SourceView->SourceStructuredBuffer || UnorderedAccessView->SourceView->SourceVertexBuffer || UnorderedAccessView->SourceView->SourceIndexBuffer || (Surface && Surface->Texture.GetBuffer()))
 	{
-		check(UnorderedAccessView->SourceView->SourceStructuredBuffer || UnorderedAccessView->SourceView->SourceVertexBuffer || UnorderedAccessView->SourceView->SourceIndexBuffer);
+		check(UnorderedAccessView->SourceView->SourceStructuredBuffer || UnorderedAccessView->SourceView->SourceVertexBuffer || UnorderedAccessView->SourceView->SourceIndexBuffer || (Surface && Surface->Texture.GetBuffer()));
 		
 		FMetalBuffer Buffer;
 		uint32 Size = 0;
@@ -517,6 +518,10 @@ void FMetalRHICommandContext::RHIClearTinyUAV(FUnorderedAccessViewRHIParamRef Un
 			Buffer = UnorderedAccessView->SourceView->SourceIndexBuffer->Buffer;
 			Size = UnorderedAccessView->SourceView->SourceIndexBuffer->GetSize();
 		}
+		else if (Surface && Surface->Texture.GetBuffer())
+		{
+			Buffer = FMetalBuffer(Surface->Texture.GetBuffer(), false);
+		}
 		
 		uint32 NumComponents = 1;
 		uint32 NumBytes = 1;
