Fix Metal 1.2 compilation - post-IAB changes we need to print the inner-type of device buffers directly in some cases, as the normal type printing will emit a buffer_argument structure for the IABs.

diff --git a/Engine/Source/Developer/Apple/MetalShaderFormat/Private/MetalBackend.cpp b/Engine/Source/Developer/Apple/MetalShaderFormat/Private/MetalBackend.cpp
--- a/Engine/Source/Developer/Apple/MetalShaderFormat/Private/MetalBackend.cpp
+++ b/Engine/Source/Developer/Apple/MetalShaderFormat/Private/MetalBackend.cpp
@@ -1151,7 +1151,7 @@ protected:
 													   buffer,
 													   "const device "
 													   );
-								print_type_pre(PtrType);
+								print_base_type(PtrType->inner_type);
 								ralloc_asprintf_append(buffer, " *%s", unique_name(var));
 								print_type_post(PtrType);
 								ralloc_asprintf_append(
