Support lod argument on SRV .Load operations - note that hlslcc DOES NOT support lod arguments on UAV Load/Store operations so far as I can see.

diff --git a/Engine/Source/Developer/Apple/MetalShaderFormat/Private/MetalBackend.cpp b/Engine/Source/Developer/Apple/MetalShaderFormat/Private/MetalBackend.cpp
--- a/Engine/Source/Developer/Apple/MetalShaderFormat/Private/MetalBackend.cpp
+++ b/Engine/Source/Developer/Apple/MetalShaderFormat/Private/MetalBackend.cpp
@@ -2294,6 +2294,11 @@ protected:
 					ralloc_asprintf_append(buffer, ",");
 					tex->lod_info.sample_index->accept(this);
 				}
+				else if (tex->lod_info.lod && !tex->lod_info.lod->is_zero())
+				{
+					ralloc_asprintf_append(buffer, ",");
+					tex->lod_info.lod->accept(this);
+				}
 			}
 		}
 			break;
