Sort the Metal shader arguments so that all resources are declared by index, this fixes incorrect rendering on Nvidia with the GPUScene enabled though it is unclear why...

diff --git a/Engine/Source/Developer/Apple/MetalShaderFormat/Private/MetalBackend.cpp b/Engine/Source/Developer/Apple/MetalShaderFormat/Private/MetalBackend.cpp
--- a/Engine/Source/Developer/Apple/MetalShaderFormat/Private/MetalBackend.cpp
+++ b/Engine/Source/Developer/Apple/MetalShaderFormat/Private/MetalBackend.cpp
@@ -1516,20 +1516,51 @@ protected:
 			);
 			bPrintComma = true;
 		}
-
+		
+		TArray<TArray<ir_variable *>> SortedParams;
 		foreach_iter(exec_list_iterator, iter, sig->parameters)
 		{
 			ir_variable *const inst = (ir_variable *) iter.get();
-			if (bPrintComma)
+			int32 index = Buffers.GetIndex(inst);
+			if (index > 0)
 			{
-				ralloc_asprintf_append(buffer, ",\n");
-				++indentation;
-				indent();
-				--indentation;
+				if (index >= SortedParams.Num())
+				{
+					SortedParams.AddDefaulted((index + 1) - SortedParams.Num());
+				}
+				TArray<ir_variable *>& Set = SortedParams[index];
+				Set.Add(inst);
+			}
+			else
+			{
+				if (bPrintComma)
+				{
+					ralloc_asprintf_append(buffer, ",\n");
+					++indentation;
+					indent();
+					--indentation;
+				}
+				inst->accept(this);
+				bPrintComma = true;
 			}
-			inst->accept(this);
-			bPrintComma = true;
 		}
+		
+		for (auto Pair : SortedParams)
+		{
+			for (auto inst : Pair)
+			{
+				if (bPrintComma)
+				{
+					ralloc_asprintf_append(buffer, ",\n");
+					++indentation;
+					indent();
+					--indentation;
+				}
+				inst->accept(this);
+				bPrintComma = true;
+			}
+		}
+		
 		check(sig->is_main);
 		ralloc_asprintf_append(buffer, ")\n");
 
diff --git a/Engine/Source/Developer/Apple/MetalShaderFormat/Private/MetalShaderFormat.cpp b/Engine/Source/Developer/Apple/MetalShaderFormat/Private/MetalShaderFormat.cpp
--- a/Engine/Source/Developer/Apple/MetalShaderFormat/Private/MetalShaderFormat.cpp
+++ b/Engine/Source/Developer/Apple/MetalShaderFormat/Private/MetalShaderFormat.cpp
@@ -239,7 +239,7 @@ class FMetalShaderFormat : public IShaderFormat
 public:
 	enum
 	{
-		HEADER_VERSION = 60,
+		HEADER_VERSION = 61,
 	};
 	
 	struct FVersion
