Fix for side table parsing when using tessellation - consume new line otherwise the rest of the tessellation data is lost.

diff --git a/Engine/Source/Developer/Apple/MetalShaderFormat/Private/MetalShaderCompiler.cpp b/Engine/Source/Developer/Apple/MetalShaderFormat/Private/MetalShaderCompiler.cpp
--- a/Engine/Source/Developer/Apple/MetalShaderFormat/Private/MetalShaderCompiler.cpp
+++ b/Engine/Source/Developer/Apple/MetalShaderFormat/Private/MetalShaderCompiler.cpp
@@ -753,6 +753,12 @@ static const int32 Str##PrefixLen = FCStringAnsi::Strlen(Str##Prefix)
 				ShaderSource++;
 			}
 		}
+		
+		if (*ShaderSource && !CrossCompiler::Match(ShaderSource, '\n'))
+		{
+			return false;
+		}
+		
 		if (SideTable < 0)
 		{
 			UE_LOG(LogMetalShaderCompiler, Fatal, TEXT("Couldn't parse the SideTable buffer index for bounds checking"));
