DR - Dump hlslcc flags for Metal direct compiler
- Don't validate on Vulkan shaders unless on debug

diff --git a/Engine/Source/Developer/Apple/MetalShaderFormat/Private/MetalShaderCompiler.cpp b/Engine/Source/Developer/Apple/MetalShaderFormat/Private/MetalShaderCompiler.cpp
--- a/Engine/Source/Developer/Apple/MetalShaderFormat/Private/MetalShaderCompiler.cpp
+++ b/Engine/Source/Developer/Apple/MetalShaderFormat/Private/MetalShaderCompiler.cpp
@@ -2254,6 +2254,16 @@ void CompileShader_Metal(const FShaderCompilerInput& _Input,FShaderCompilerOutpu
         RemoveUniformBuffersFromSource(Input.Environment, PreprocessedShader);
     }
 
+	uint32 CCFlags = HLSLCC_NoPreprocess | HLSLCC_PackUniformsIntoUniformBufferWithNames | HLSLCC_FixAtomicReferences | HLSLCC_KeepSamplerAndImageNames;
+	if (!bDirectCompile || UE_BUILD_DEBUG)
+	{
+		// Validation is expensive - only do it when compiling directly for debugging
+		CCFlags |= HLSLCC_NoValidation;
+	}
+
+	// Required as we added the RemoveUniformBuffersFromSource() function (the cross-compiler won't be able to interpret comments w/o a preprocessor)
+	CCFlags &= ~HLSLCC_NoPreprocess;
+
 	// Write out the preprocessed file and a batch file to compile it if requested (DumpDebugInfoPath is valid)
 	if (bDumpDebugInfo && !bDirectCompile)
 	{
@@ -2279,17 +2289,10 @@ void CompileShader_Metal(const FShaderCompilerInput& _Input,FShaderCompilerOutpu
 
 		if (Input.bGenerateDirectCompileFile)
 		{
-			FFileHelper::SaveStringToFile(CreateShaderCompilerWorkerDirectCommandLine(Input), *(Input.DumpDebugInfoPath / TEXT("DirectCompile.txt")));
+			FFileHelper::SaveStringToFile(CreateShaderCompilerWorkerDirectCommandLine(Input, CCFlags), *(Input.DumpDebugInfoPath / TEXT("DirectCompile.txt")));
 		}
 	}
 
-	uint32 CCFlags = HLSLCC_NoPreprocess | HLSLCC_PackUniformsIntoUniformBufferWithNames | HLSLCC_FixAtomicReferences | HLSLCC_KeepSamplerAndImageNames;
-	if (!bDirectCompile || UE_BUILD_DEBUG)
-	{
-		// Validation is expensive - only do it when compiling directly for debugging
-		CCFlags |= HLSLCC_NoValidation;
-	}
-	
 	FSHAHash GUIDHash;
 	if (!bDirectCompile)
 	{
@@ -2303,9 +2306,6 @@ void CompileShader_Metal(const FShaderCompilerInput& _Input,FShaderCompilerOutpu
 		FGuid Guid = FGuid::NewGuid();
 		FSHA1::HashBuffer(&Guid, sizeof(FGuid), GUIDHash.Hash);
 	}
-	
-	// Required as we added the RemoveUniformBuffersFromSource() function (the cross-compiler won't be able to interpret comments w/o a preprocessor)
-	CCFlags &= ~HLSLCC_NoPreprocess;
 
 	FMetalShaderOutputCooker* Cooker = new FMetalShaderOutputCooker(Input,Output,WorkingDirectory, PreprocessedShader, GUIDHash, VersionEnum, CCFlags, HlslCompilerTarget, MetalCompilerTarget, Semantics, TypeMode, MaxUnrollLoops, Frequency, bDumpDebugInfo, Standard, MinOSVersion);
 		
diff --git a/Engine/Source/Developer/VulkanShaderFormat/Private/VulkanShaderCompiler.cpp b/Engine/Source/Developer/VulkanShaderFormat/Private/VulkanShaderCompiler.cpp
--- a/Engine/Source/Developer/VulkanShaderFormat/Private/VulkanShaderCompiler.cpp
+++ b/Engine/Source/Developer/VulkanShaderFormat/Private/VulkanShaderCompiler.cpp
@@ -1908,7 +1908,8 @@ void DoCompileVulkanShader(const FShaderCompilerInput& Input, FShaderCompilerOut
 
 	// Preprocess the shader.
 	FString PreprocessedShaderSource;
-	if (Input.bSkipPreprocessedCache)
+	const bool bDirectCompile = FParse::Param(FCommandLine::Get(), TEXT("directcompile"));
+	if (bDirectCompile)
 	{
 		if (!FFileHelper::LoadFileToString(PreprocessedShaderSource, *Input.VirtualSourceFilePath))
 		{
@@ -1965,6 +1966,12 @@ void DoCompileVulkanShader(const FShaderCompilerInput& Input, FShaderCompilerOut
 	// Required as we added the RemoveUniformBuffersFromSource() function (the cross-compiler won't be able to interpret comments w/o a preprocessor)
 	CompilerInfo.CCFlags &= ~HLSLCC_NoPreprocess;
 
+	if (!bDirectCompile || UE_BUILD_DEBUG)
+	{
+		// Validation is expensive - only do it when compiling directly for debugging
+		CompilerInfo.CCFlags |= HLSLCC_NoValidation;
+	}
+
 	// Write out the preprocessed file and a batch file to compile it if requested (DumpDebugInfoPath is valid)
 	if (CompilerInfo.bDebugDump)
 	{
@@ -2006,14 +2013,14 @@ void DoCompileVulkanShader(const FShaderCompilerInput& Input, FShaderCompilerOut
 			auto* SourceWithHeader = GeneratedGlslSource.GetData();
 			char* SourceNoHeader = strstr(SourceWithHeader, "#version");
 			bool bSuccess = CompileUsingInternal(CompilerInfo, BindingTable, GeneratedGlslSource, EntryPointName, Output, bHasRealUBs);
-			if (Input.bSkipPreprocessedCache)
+			if (bDirectCompile)
 			{
 				FPlatformMisc::LowLevelOutputDebugStringf(TEXT("Success: %d\n%s\n"), bSuccess, ANSI_TO_TCHAR(SourceWithHeader));
 			}
 		}
 	}
 	
-	if (Input.bSkipPreprocessedCache)
+	if (bDirectCompile)
 	{
 		for (const auto& Error : Output.Errors)
 		{
