Remove all support for Metal prior to Metal 1.2 on iOS/tvOS and Metal 2.0 on macOS. This permits some fairly radical simplifications and removal of a number of workarounds.

- All Buffer<>/RWBuffer<> objects are now bound as texture2d or texture_buffer, unless they are 3-component, in which case they are bound as raw T* buffer pointers.
- 3-component Buffer/RWBuffer objects should now work - Metal needed to use packed_* types internally as normal 3-component types in Metal are the same size as 4-component types which is different to HLSL. This is untested as we have no such example.
- Texture Gather/GatherCompare operations have been fixed and should now work.
- Remove all the TypedBufferFormat function-constant specialisation infrastructure and metadata as it is no longer necessary.
- Remove all the raw-conversion for 1, 2 and 4 component vector types as it is no longer required.
- Fixed input qualifiers (centroid, linear, etc.) for Fragment shader inputs - they should now work.
- Remove a bunch of Metal specific shader changes due to buffer type casting and swizzling that will now just work.
- Use the BGRA vertex format for FColor on Metal when available (everywhere but iOS/tvOS 10 using Metal 1.2).
- Simplify ShouldCompilePermutation & similar functions now that we have fewer versions to support.

diff --git a/Engine/Shaders/Private/MaterialTemplate.ush b/Engine/Shaders/Private/MaterialTemplate.ush
--- a/Engine/Shaders/Private/MaterialTemplate.ush
+++ b/Engine/Shaders/Private/MaterialTemplate.ush
@@ -2159,7 +2159,7 @@ half3x3 AssembleTangentToWorld( half3 TangentToWorld0, half4 TangentToWorld2 )
 	// Use conservative depth output so we still get Z Cull.  Note, this is a reversed Z depth surface.
 	#define DEPTH_WRITE_SEMANTIC SV_DepthLessEqual
 #elif COMPILER_METAL
-	#define INPUT_POSITION_QUALIFIERS linear centroid
+	#define INPUT_POSITION_QUALIFIERS 
 	#define DEPTH_WRITE_SEMANTIC SV_DepthLessEqual
 #elif PS4_PROFILE
 	#define INPUT_POSITION_QUALIFIERS
diff --git a/Engine/Shaders/Private/ParticleBoundsShader.usf b/Engine/Shaders/Private/ParticleBoundsShader.usf
--- a/Engine/Shaders/Private/ParticleBoundsShader.usf
+++ b/Engine/Shaders/Private/ParticleBoundsShader.usf
@@ -22,13 +22,7 @@
 #define MIN_FLT (-3.402823466e+38f)
 
 /** Input buffer containing particle indices. */
-#if COMPILER_METAL
-#undef half2
-#define vec2 half2
-#else
-#define vec2 float2
-#endif
-Buffer<vec2> InParticleIndices;
+Buffer<float2> InParticleIndices;
 /** Texture containing particle positions. */
 Texture2D PositionTexture;
 
diff --git a/Engine/Shaders/Private/ParticleGPUSpriteVertexFactory.ush b/Engine/Shaders/Private/ParticleGPUSpriteVertexFactory.ush
--- a/Engine/Shaders/Private/ParticleGPUSpriteVertexFactory.ush
+++ b/Engine/Shaders/Private/ParticleGPUSpriteVertexFactory.ush
@@ -17,13 +17,7 @@
 #define USE_PARTICLE_SIZE (FEATURE_LEVEL >= FEATURE_LEVEL_SM5 && NEEDS_PARTICLE_SIZE)
 
 /** Buffer containing particle indices. */
-#if COMPILER_METAL
-	// Metal does not have buffer SRVs, so data types have to match
-	#undef half2
-	Buffer<half2> ParticleIndices;
-#else
-	Buffer<float2> ParticleIndices;
-#endif
+Buffer<float2> ParticleIndices;
 
 /** Offset in to the particle indices buffer. */
 uint ParticleIndicesOffset;
diff --git a/Engine/Shaders/Private/ParticleSortKeyGen.usf b/Engine/Shaders/Private/ParticleSortKeyGen.usf
--- a/Engine/Shaders/Private/ParticleSortKeyGen.usf
+++ b/Engine/Shaders/Private/ParticleSortKeyGen.usf
@@ -14,20 +14,14 @@
 #include "Common.ush"
 
 /** Input buffer containing particle indices. */
-#if COMPILER_METAL
-#undef half2
-#define vec2 half2
-#else
-#define vec2 float2
-#endif
-Buffer<vec2> InParticleIndices;
+Buffer<float2> InParticleIndices;
 /** Texture containing particle positions. */
 Texture2D PositionTexture;
 
 /** Output key buffer. */
 RWBuffer<uint> OutKeys;
 /** Output indices buffer. */
-RWBuffer<vec2> OutParticleIndices;
+RWBuffer<float2> OutParticleIndices;
 
 [numthreads(THREAD_COUNT,1,1)]
 void GenerateParticleSortKeys(
@@ -61,7 +55,7 @@ void GenerateParticleSortKeys(
 		if ( InputIndex < ParticleKeyGen.KeyCount )
 		{
 			// Read in the particle index and its position.
-			const vec2 ParticleIndex = InParticleIndices[InputIndex];
+			const float2 ParticleIndex = InParticleIndices[InputIndex];
 			int3 ParticleTexel = int3(ParticleIndex.xy * int2(TEXTURE_SIZE_X, TEXTURE_SIZE_Y), 0);
 			const float4 ParticlePosition = PositionTexture.Load(ParticleTexel);
 
diff --git a/Engine/Shaders/Private/SeparateTranslucency.ush b/Engine/Shaders/Private/SeparateTranslucency.ush
--- a/Engine/Shaders/Private/SeparateTranslucency.ush
+++ b/Engine/Shaders/Private/SeparateTranslucency.ush
@@ -28,7 +28,7 @@ float4 BilinearUpsampling(float2 UV, Texture2D LowResTex)
 float4 NearestDepthNeighborUpsampling(float2 Position, float2 UV, Texture2D LowResTex, float2 LowResTexelSize)
 {
 //@todo - support for all platforms, just skip the GatherRed optimization where not supported
-#if (SM5_PROFILE && !(METAL_SM5_PROFILE || METAL_SM5_NOTESS_PROFILE || METAL_MRT_PROFILE || PS4_PROFILE))
+#if (SM5_PROFILE && !(PS4_PROFILE))
 
 	// The relative depth comparison breaks down at larger distances and incorrectly causes point sampling on the skybox pixels
 	float MaxOperationDepth = 2000000.0f;
diff --git a/Engine/Shaders/Private/SimpleElementHitProxyPixelShader.usf b/Engine/Shaders/Private/SimpleElementHitProxyPixelShader.usf
--- a/Engine/Shaders/Private/SimpleElementHitProxyPixelShader.usf
+++ b/Engine/Shaders/Private/SimpleElementHitProxyPixelShader.usf
@@ -17,11 +17,7 @@ void Main(
 	)
 {
 	float alpha = Texture2DSample(InTexture, InTextureSampler,TextureCoordinate).a;
-#if COMPILER_METAL
-	OutColor = float4(HitProxyId.zyx,alpha);
-#else
 	OutColor = float4(HitProxyId.rgb,alpha);
-#endif
 
 #if FEATURE_LEVEL >= FEATURE_LEVEL_SM4
 	// Dx11 (feature levels 10 and upwards) do not perform alpha test, so manually clip.
diff --git a/Engine/Shaders/Public/Platform.ush b/Engine/Shaders/Public/Platform.ush
--- a/Engine/Shaders/Public/Platform.ush
+++ b/Engine/Shaders/Public/Platform.ush
@@ -241,7 +241,11 @@
 	// @todo metal: remove this and make sure all uses handle METAL_PROFILE
 	#undef ES3_1_PROFILE
 	#define ES3_1_PROFILE 1
+#if MAX_SHADER_LANGUAGE_VERSION >= 3
+	#define FCOLOR_COMPONENT_SWIZZLE .rgba
+#else
 	#define FCOLOR_COMPONENT_SWIZZLE .bgra
+#endif
 	#define FMANUALFETCH_COLOR_COMPONENT_SWIZZLE .bgra
 	#define STENCIL_COMPONENT_SWIZZLE .x
 
@@ -250,25 +254,33 @@
 	// @todo metal: remove this and make sure all uses handle METAL_ES2_PROFILE
 	#undef ES2_PROFILE
 	#define ES2_PROFILE 1
+#if MAX_SHADER_LANGUAGE_VERSION >= 3
+	#define FCOLOR_COMPONENT_SWIZZLE .rgba
+#else
 	#define FCOLOR_COMPONENT_SWIZZLE .bgra
+#endif
 	#define FMANUALFETCH_COLOR_COMPONENT_SWIZZLE .bgra
 	#define STENCIL_COMPONENT_SWIZZLE .x
 
 #elif METAL_MRT_PROFILE
 	#define FEATURE_LEVEL FEATURE_LEVEL_SM5
+#if MAX_SHADER_LANGUAGE_VERSION >= 3
+	#define FCOLOR_COMPONENT_SWIZZLE .rgba
+#else
 	#define FCOLOR_COMPONENT_SWIZZLE .bgra
+#endif
 	#define FMANUALFETCH_COLOR_COMPONENT_SWIZZLE .bgra
 	#define STENCIL_COMPONENT_SWIZZLE .x
 
 #elif METAL_SM5_NOTESS_PROFILE
 	#define FEATURE_LEVEL FEATURE_LEVEL_SM5
-	#define FCOLOR_COMPONENT_SWIZZLE .bgra
+	#define FCOLOR_COMPONENT_SWIZZLE .rgba
 	#define FMANUALFETCH_COLOR_COMPONENT_SWIZZLE .bgra
 	#define STENCIL_COMPONENT_SWIZZLE .x
 
 #elif METAL_SM5_PROFILE
 	#define FEATURE_LEVEL FEATURE_LEVEL_SM5
-	#define FCOLOR_COMPONENT_SWIZZLE .bgra
+	#define FCOLOR_COMPONENT_SWIZZLE .rgba
 	#define FMANUALFETCH_COLOR_COMPONENT_SWIZZLE .bgra
 	#define STENCIL_COMPONENT_SWIZZLE .x
 
@@ -532,7 +544,7 @@
 
 // ---------------------------------------------------- Compiler missing implementations
 
-#if COMPILER_GLSL_ES2 || (COMPILER_METAL && MAX_SHADER_LANGUAGE_VERSION < 2) || COMPILER_SWITCH
+#if COMPILER_GLSL_ES2 || COMPILER_SWITCH
 
 	float determinant(float3x3 M)
 	{
diff --git a/Engine/Shaders/Public/Platform/Metal/MetalCommon.ush b/Engine/Shaders/Public/Platform/Metal/MetalCommon.ush
--- a/Engine/Shaders/Public/Platform/Metal/MetalCommon.ush
+++ b/Engine/Shaders/Public/Platform/Metal/MetalCommon.ush
@@ -15,16 +15,8 @@
 #endif
 #endif
 
-#if MAX_SHADER_LANGUAGE_VERSION >= 3
-#define STRONG_TYPE 
+#define STRONG_TYPE
 #define StrongTypedBuffer Buffer
-#else
-// Abuse the GLSL-style "invariant" keyword that hlslcc exposes to modify Buffer<>/RWBuffer<> declarations to assume that they input type is always the same as that declared in the shader
-// The MetalRHI runtime will then assert when it isn't. For Metal 1.1 and below all Buffer<>/RWBuffer<> objects will *assume* this type matches as there is no way to gracefully handle it there.
-#define STRONG_TYPE invariant
-// Note: this is assuming that Metal Buffers and StructuredBuffers are interchangeable
-#define StrongTypedBuffer StructuredBuffer
-#endif
 
 /* Too many bugs to enable wave-intrinsics just yet, but we are getting there.
 #if ((MAX_SHADER_LANGUAGE_VERSION >= 3 && COMPUTESHADER) || (MAX_SHADER_LANGUAGE_VERSION >= 4 && PIXELSHADER))
diff --git a/Engine/Source/Developer/Apple/MetalShaderFormat/Private/MetalBackend.cpp b/Engine/Source/Developer/Apple/MetalShaderFormat/Private/MetalBackend.cpp
--- a/Engine/Source/Developer/Apple/MetalShaderFormat/Private/MetalBackend.cpp
+++ b/Engine/Source/Developer/Apple/MetalShaderFormat/Private/MetalBackend.cpp
@@ -587,7 +587,6 @@ protected:
 	bool bInsertSideTable;
 
 	bool bRequiresWave;
-	bool bInsertUnknownFormat;
 	bool bNeedsDeviceIndex;
 	
     const char *shaderPrefix()
@@ -1437,7 +1436,7 @@ protected:
 		scope_depth++;
 		IsMain = sig->is_main;
 
-		if (sig->is_main && sig->is_early_depth_stencil && Frequency == fragment_shader && Backend.Version >= 2)
+		if (sig->is_main && sig->is_early_depth_stencil && Frequency == fragment_shader)
 		{
 			bExplicitEarlyFragTests = true;
 		}
@@ -2257,15 +2256,6 @@ protected:
 					tex->sampler->accept(this);
 					ralloc_asprintf_append(buffer, ", ");
 					tex->coordinate->accept(this);
-					if (FormatIndex >= 0)
-					{
-						ralloc_asprintf_append(buffer, ", GMetalTypedBufferFormat%d", Index);
-					}
-					else
-					{
-						ralloc_asprintf_append(buffer, ", GMetalUnknownFormat");
-						bInsertUnknownFormat = true;
-					}
 					if (bSideTable)
 					{
 						ralloc_asprintf_append(buffer, ", %s)", BufferSizesName);
@@ -2277,8 +2267,6 @@ protected:
 				}
 				else if (Backend.bBoundsChecks)
 				{
-					check(Index <= 30);
-					
 					if (!bIsAtomic && (!bIsStructuredBuffer || !Texture->type->inner_type->is_record()))
 					{
 						ralloc_asprintf_append(buffer, "buffer::load<");
@@ -2434,6 +2422,31 @@ protected:
 				ralloc_asprintf_append(buffer, ", ");
 				tex->offset->accept(this);
 			}
+			else if(tex->channel > ir_channel_none)
+			{
+				ralloc_asprintf_append(buffer, ", int2(0)");
+			}
+
+			// Emit channel selection for gather
+			check(tex->channel < ir_channel_unknown);
+			switch(tex->channel)
+			{
+				case ir_channel_red:
+					ralloc_asprintf_append(buffer, ", component::x");
+					break;
+				case ir_channel_green:
+					ralloc_asprintf_append(buffer, ", component::y");
+					break;
+				case ir_channel_blue:
+					ralloc_asprintf_append(buffer, ", component::z");
+					break;
+				case ir_channel_alpha:
+					ralloc_asprintf_append(buffer, ", component::w");
+					break;
+				case ir_channel_none:
+				default:
+					break;
+			}
 		}
 			break;
 
@@ -2636,20 +2649,9 @@ protected:
 			ralloc_asprintf_append(buffer, "[");
 		}
 		
-		bool bIsVectorArrayIndex = deref->array->type->is_vector() && (Backend.Version < 3 && Backend.bIsDesktop == EMetalGPUSemanticsImmediateDesktop);
-		if (bIsVectorArrayIndex)
-		{
-			ralloc_asprintf_append(buffer, "vector_array_deref(");
-		}
-
 		deref->array_index->accept(this);
 		should_print_uint_literals_as_ints = false;
 
-		if (bIsVectorArrayIndex)
-		{
-			ralloc_asprintf_append(buffer, ")");
-		}
-
 		if (enforceInt)
 		{
 			ralloc_asprintf_append(buffer, ")]");
@@ -2743,15 +2745,6 @@ protected:
 						deref->image->accept(this);
 						ralloc_asprintf_append(buffer, ", ");
 						deref->image_index->accept(this);
-						if (FormatIndex >= 0)
-						{
-							ralloc_asprintf_append(buffer, ", GMetalTypedBufferFormat%d", Index);
-						}
-						else
-						{
-							ralloc_asprintf_append(buffer, ", GMetalUnknownFormat");
-							bInsertUnknownFormat = true;
-						}
 						if (bSideTable)
 						{
 							ralloc_asprintf_append(buffer, ", %s)", BufferSizesName);
@@ -2941,15 +2934,6 @@ protected:
 						deref->image->accept(this);
 						ralloc_asprintf_append(buffer, ", ");
 						deref->image_index->accept(this);
-						if (FormatIndex >= 0)
-						{
-							ralloc_asprintf_append(buffer, ", GMetalTypedBufferFormat%d", Index);
-						}
-						else
-						{
-							ralloc_asprintf_append(buffer, ", GMetalUnknownFormat");
-							bInsertUnknownFormat = true;
-						}
 						if (bSideTable)
 						{
 							ralloc_asprintf_append(buffer, ", %s, ", BufferSizesName);
@@ -4620,10 +4604,9 @@ public:
 		, bUsePacked(false)
 		, bNeedsComputeInclude(false)
 		, bExplicitEarlyFragTests(false)
-		, bImplicitEarlyFragTests(InBackend.Version >= 2)
+		, bImplicitEarlyFragTests(true)
 		, bInsertSideTable(false)
 		, bRequiresWave(false)
-		, bInsertUnknownFormat(false)
 		, bNeedsDeviceIndex(false)
 	{
 		printable_names = hash_table_ctor(32, hash_table_pointer_hash, hash_table_pointer_compare);
@@ -4660,39 +4643,8 @@ public:
 		char* decl_buffer = ralloc_asprintf(mem_ctx, "");
 		buffer = &decl_buffer;
 		declare_structs(ParseState);
-        
-        if (Backend.TypedBuffers)
-        {
-            ralloc_asprintf_append(buffer, "\n");
-            for (uint32 i = 0; i < (uint32)Backend.TypedBufferFormats.Num(); i++)
-            {
-                if ((Backend.TypedBuffers & (1 << i)) != 0)
-                {
-                    if ((Backend.TypedUAVs & (1 << i)) != 0)
-                    {
-                        ralloc_asprintf_append(buffer, "#if __METAL_TYPED_BUFFER_RW_IMPL__ == __METAL_TYPED_BUFFER_RAW__\n");
-                        ralloc_asprintf_append(buffer, "constant uint GMetalTypedBufferFormat%d = ue4::Max;\n", i);
-                        ralloc_asprintf_append(buffer, "#elif __METAL_TYPED_BUFFER_RW_IMPL__ == __METAL_TYPED_BUFFER_FC__\n");
-                        ralloc_asprintf_append(buffer, "constant uint GMetalTypedBufferFormat%d [[ function_constant(%d) ]];\n", i, i);
-                        ralloc_asprintf_append(buffer, "#else\n");
-                        ralloc_asprintf_append(buffer, "constant uint GMetalTypedBufferFormat%d = ue4::Unknown;\n", i);
-                        ralloc_asprintf_append(buffer, "#endif\n");
-                    }
-                    else
-                    {
-                        ralloc_asprintf_append(buffer, "#if __METAL_TYPED_BUFFER_READ_IMPL__ == __METAL_TYPED_BUFFER_RAW__\n");
-                        ralloc_asprintf_append(buffer, "constant uint GMetalTypedBufferFormat%d = ue4::Max;\n", i);
-                        ralloc_asprintf_append(buffer, "#elif __METAL_TYPED_BUFFER_READ_IMPL__ == __METAL_TYPED_BUFFER_FC__\n");
-                        ralloc_asprintf_append(buffer, "constant uint GMetalTypedBufferFormat%d [[ function_constant(%d) ]];\n", i, i);
-                        ralloc_asprintf_append(buffer, "#else\n");
-                        ralloc_asprintf_append(buffer, "constant uint GMetalTypedBufferFormat%d = ue4::Unknown;\n", i);
-                        ralloc_asprintf_append(buffer, "#endif\n");
-                    }
-                }
-            }
-        }
 		
-        if ((bExplicitEarlyFragTests || bImplicitEarlyFragTests) && !Backend.bExplicitDepthWrites && Frequency == fragment_shader && Backend.Version >= 2)
+        if ((bExplicitEarlyFragTests || bImplicitEarlyFragTests) && !Backend.bExplicitDepthWrites && Frequency == fragment_shader)
 		{
 			ralloc_asprintf_append(buffer, "\n#define FUNC_ATTRIBS [[early_fragment_tests]]\n\n");
 		}
@@ -4700,11 +4652,6 @@ public:
 		{
 			ralloc_asprintf_append(buffer, "\n#define FUNC_ATTRIBS \n\n");
 		}
-		
-		if (bInsertUnknownFormat)
-		{
-			ralloc_asprintf_append(buffer, "\nconstant uint GMetalUnknownFormat = ue4::Unknown;\n\n");
-		}
 
 		// These should work in fragment shaders but Apple are behind the curve on SM6.
 		if (bRequiresWave && Frequency == compute_shader && Backend.Version >= 3)
@@ -4787,22 +4734,14 @@ public:
                 ralloc_asprintf_append(buffer, "#define __METAL_TYPED_BUFFER_READ_IMPL__ 0\n");
                 ralloc_asprintf_append(buffer, "#define __METAL_TYPED_BUFFER_RW_IMPL__ 0\n");
                 break;
-            case EMetalTypeBufferModeSRV:
-                ralloc_asprintf_append(buffer, "#define __METAL_TYPED_BUFFER_READ_IMPL__ 1\n");
-                ralloc_asprintf_append(buffer, "#define __METAL_TYPED_BUFFER_RW_IMPL__ 2\n");
-                break;
-            case EMetalTypeBufferModeUAV:
+            case EMetalTypeBufferMode2D:
                 ralloc_asprintf_append(buffer, "#define __METAL_TYPED_BUFFER_READ_IMPL__ 1\n");
                 ralloc_asprintf_append(buffer, "#define __METAL_TYPED_BUFFER_RW_IMPL__ 1\n");
                 break;
-            case EMetalTypeBufferModeTex:
+            case EMetalTypeBufferModeTB:
                 ralloc_asprintf_append(buffer, "#define __METAL_TYPED_BUFFER_READ_IMPL__ 3\n");
                 ralloc_asprintf_append(buffer, "#define __METAL_TYPED_BUFFER_RW_IMPL__ 3\n");
                 break;
-            case EMetalTypeBufferModeFun:
-                ralloc_asprintf_append(buffer, "#define __METAL_TYPED_BUFFER_READ_IMPL__ 2\n");
-                ralloc_asprintf_append(buffer, "#define __METAL_TYPED_BUFFER_RW_IMPL__ 2\n");
-                break;
             default:
                 break;
         }
@@ -5065,12 +5004,6 @@ struct FMetalCheckComputeRestrictionsVisitor : public ir_rvalue_visitor
 			{
 				ImageRW[Var] |= EMetalAccessRead;
 			}
-
-			if (ImageRW[Var] == EMetalAccessReadWrite && Version < 2 && (!Var->type->sampler_buffer || TypeMode == EMetalTypeBufferModeUAV))
-			{
-				_mesa_glsl_error(ParseState, "Metal doesn't allow simultaneous read & write on RWTexture(s) %s%s%s", Var->name ? "(" : "", Var->name ? Var->name : "", Var->name ? ")" : "");
-				bErrors = true;
-			}
 		}
 	}
 
@@ -5110,19 +5043,6 @@ struct FMetalCheckNonComputeRestrictionsVisitor : public FMetalCheckComputeRestr
 	
 	virtual ir_visitor_status visit(ir_variable* IR) override
 	{
-		if (IR->type && IR->type->is_image() && (Version < 2))
-		{
-			if (IR->name)
-			{
-				_mesa_glsl_error(ParseState, "Metal doesn't allow UAV '%s' on non-compute shader stage %d.", IR->name, (unsigned)ParseState->target);
-			}
-			else
-			{
-				_mesa_glsl_error(ParseState, "Metal doesn't allow UAV on non-compute shader stage %d.", (unsigned)ParseState->target);
-			}
-			bErrors = true;
-			return visit_stop;
-		}
 		// @todo validate that GLSL_OUTPUTTOPOLOGY_POINT, GLSL_OUTPUTTOPOLOGY_LINE are not used
 		
 		return FMetalCheckComputeRestrictionsVisitor::visit(IR);
@@ -5496,6 +5416,9 @@ bool FMetalCodeBackend::GenerateMain(EHlslShaderFrequency Frequency, const char*
 		}
 		else if (Variable->semantic != NULL || Variable->type->is_record())
 		{
+			Qualifier.Fields.bCentroid = Variable->centroid;
+			Qualifier.Fields.InterpolationMode = Variable->interpolation;
+			
 			ir_dereference_variable* ArgVarDeref = NULL;
 			switch (Variable->mode)
 			{
@@ -5505,6 +5428,7 @@ bool FMetalCodeBackend::GenerateMain(EHlslShaderFrequency Frequency, const char*
 					ParseState,
 					Variable->name,
 					Variable->semantic,
+					Qualifier,
 					Variable->type,
 					&DeclInstructions,
 					&PreCallInstructions
@@ -5655,6 +5579,9 @@ bool FMetalCodeBackend::GenerateMain(EHlslShaderFrequency Frequency, const char*
 				const ir_variable* Variable = (ir_variable*)Iter.get();
 				if (Variable->semantic != NULL || Variable->type->is_record())
 				{
+					Qualifier.Fields.bCentroid = Variable->centroid;
+					Qualifier.Fields.InterpolationMode = Variable->interpolation;
+					
 					ir_dereference_variable* ArgVarDeref = NULL;
 					switch (Variable->mode)
 					{
@@ -5664,6 +5591,7 @@ bool FMetalCodeBackend::GenerateMain(EHlslShaderFrequency Frequency, const char*
 							ParseState,
 							Variable->name,
 							Variable->semantic,
+							Qualifier,
 							Variable->type,
 							&VertexDeclInstructions,
 							&VertexPreCallInstructions
diff --git a/Engine/Source/Developer/Apple/MetalShaderFormat/Private/MetalBackend.h b/Engine/Source/Developer/Apple/MetalShaderFormat/Private/MetalBackend.h
--- a/Engine/Source/Developer/Apple/MetalShaderFormat/Private/MetalBackend.h
+++ b/Engine/Source/Developer/Apple/MetalShaderFormat/Private/MetalBackend.h
@@ -25,15 +25,10 @@ public:
 	
 	uint32 GetClipDistanceCount() const { return ClipDistanceCount; }
 
-	virtual bool SupportsDeterminantIntrinsic() const override
-	{
-		return (Version >= 2);
-	}
-
-	virtual bool SupportsTransposeIntrinsic() const override
-	{
-		return (Version >= 2);
-	}
+	virtual bool SupportsDeterminantIntrinsic() const override { return true; }
+
+	virtual bool SupportsTransposeIntrinsic() const override { return true; }
+	
 	virtual bool SupportsIntegerModulo() const override { return true; }
 
 	virtual bool SupportsMatrixConversions() const override { return false; }
@@ -50,9 +45,9 @@ public:
 	
 	virtual bool SplitInputVariableStructs() const { return false; }
 	
-	virtual bool SupportsFusedMultiplyAdd() const { return (Version >= 2); }
+	virtual bool SupportsFusedMultiplyAdd() const { return true; }
 	
-	virtual bool SupportsSaturateIntrinsic() const { return (Version >= 2); }
+	virtual bool SupportsSaturateIntrinsic() const { return true; }
 
     virtual bool SupportsSinCosIntrinsic() const { return true; }
     
@@ -83,10 +78,8 @@ enum EMetalGPUSemantics
 enum EMetalTypeBufferMode
 {
 	EMetalTypeBufferModeRaw = 0, // No typed buffers
-    EMetalTypeBufferModeSRV = 1, // Buffer<> Typed via 2D textures, RWBuffer<> typed via function constants
-    EMetalTypeBufferModeUAV = 2, // Buffer<> SRVs & RWBuffer<> UAVs are typed via 2D textures
-    EMetalTypeBufferModeTex = 3, // Buffer<> SRVs & RWBuffer<> UAVs are typed via texture-buffers
-    EMetalTypeBufferModeFun = 4, // Buffer<> SRVs & RWBuffer<> UAVs are typed via function constants
+    EMetalTypeBufferMode2D = 2, // Buffer<> SRVs & RWBuffer<> UAVs are typed via 2D textures
+    EMetalTypeBufferModeTB = 3, // Buffer<> SRVs & RWBuffer<> UAVs are typed via texture-buffers
 };
 
 // Metal supports 16 across all HW
diff --git a/Engine/Source/Developer/Apple/MetalShaderFormat/Private/MetalShaderCompiler.cpp b/Engine/Source/Developer/Apple/MetalShaderFormat/Private/MetalShaderCompiler.cpp
--- a/Engine/Source/Developer/Apple/MetalShaderFormat/Private/MetalShaderCompiler.cpp
+++ b/Engine/Source/Developer/Apple/MetalShaderFormat/Private/MetalShaderCompiler.cpp
@@ -818,7 +818,7 @@ static const int32 Str##PrefixLen = FCStringAnsi::Strlen(Str##Prefix)
 	}
 	
 	// Early out for non-tessellation...
-	if (Version < 2 || !bUsingTessellation)
+	if (!bUsingTessellation)
 	{
 		return true;
 	}
@@ -1083,11 +1083,8 @@ void BuildMetalShaderOutput(
 	Header.SourceCRC = SourceCRC;
     Header.Bindings.bDiscards = false;
 	Header.Bindings.ConstantBuffers = ConstantBuffers;
-	if (Version >= 2)
 	{
-		Header.Bindings.TypedBufferFormats.SetNumZeroed(METAL_MAX_BUFFERS);
 		Header.Bindings.TypedBuffers = TypedBuffers;
-		Header.Bindings.InvariantBuffers = InvariantBuffers;
 		for (uint32 i = 0; i < (uint32)TypedBufferFormats.Num(); i++)
 		{
 			if ((TypedBuffers & (1 << i)) != 0)
@@ -1095,37 +1092,22 @@ void BuildMetalShaderOutput(
 				check(TypedBufferFormats[i] > Unknown);
                 check(TypedBufferFormats[i] < Max);
                 if ((TypeMode > EMetalTypeBufferModeRaw)
-                && (TypeMode < EMetalTypeBufferModeFun)
+                && (TypeMode <= EMetalTypeBufferModeTB)
                 && (TypedBufferFormats[i] < RGB8Sint || TypedBufferFormats[i] > RGB32Float)
-                && (TypeMode == EMetalTypeBufferModeUAV || TypeMode == EMetalTypeBufferModeTex || !(TypedUAVs & (1 << i))))
+                && (TypeMode == EMetalTypeBufferMode2D || TypeMode == EMetalTypeBufferModeTB || !(TypedUAVs & (1 << i))))
                 {
                 	Header.Bindings.LinearBuffer |= (1 << i);
 	                Header.Bindings.TypedBuffers &= ~(1 << i);
-                }
-                else
-                {
-					Header.Bindings.TypedBufferFormats[i] = TypedBufferFormats[i];
                 }
 			}
-			else if ((InvariantBuffers & (1 << i)) != 0)
-			{
-				Header.Bindings.TypedBufferFormats[i] = TypedBufferFormats[i];
-			}
 		}
 		
 		// Raw mode means all buffers are invariant
 		if (TypeMode == EMetalTypeBufferModeRaw)
 		{
 			Header.Bindings.TypedBuffers = 0;
-			Header.Bindings.InvariantBuffers = InvariantBuffers|TypedBuffers;
 		}
 	}
-	else // No typed buffers on Metal 1.0-1.1 - all are invariant
-	{
-		Header.Bindings.TypedBuffers = 0;
-		Header.Bindings.InvariantBuffers = InvariantBuffers|TypedBuffers;
-		Header.Bindings.TypedBufferFormats = TypedBufferFormats;
-	}
 	
 	FShaderParameterMap& ParameterMap = ShaderOutput.ParameterMap;
 
@@ -1684,7 +1666,6 @@ void BuildMetalShaderOutput(
 				
 #if (PLATFORM_MAC && !UNIXLIKE_TO_MAC_REMOTE_BUILDING)				
 				FString Defines = Header.bDeviceFunctionConstants ? TEXT("-D__METAL_DEVICE_CONSTANT_INDEX__=1") : TEXT("");
-				Defines += FString::Printf(TEXT(" -D__METAL_MANUAL_TEXTURE_METADATA__=%d"), !(bUsingTessellation && (Frequency == SF_Vertex || Frequency == SF_Hull)) && Version < 3);
 				Defines += FString::Printf(TEXT(" -D__METAL_USE_TEXTURE_CUBE_ARRAY__=%d"), !bIsMobile);
 				switch(TypeMode)
 				{
@@ -1692,22 +1673,14 @@ void BuildMetalShaderOutput(
 						Defines += TEXT(" -D__METAL_TYPED_BUFFER_READ_IMPL__=0");
 						Defines += TEXT(" -D__METAL_TYPED_BUFFER_RW_IMPL__=0");
 						break;
-					case EMetalTypeBufferModeSRV:
-						Defines += TEXT(" -D__METAL_TYPED_BUFFER_READ_IMPL__=1");
-						Defines += TEXT(" -D__METAL_TYPED_BUFFER_RW_IMPL__=2");
-						break;
-					case EMetalTypeBufferModeUAV:
+					case EMetalTypeBufferMode2D:
 						Defines += TEXT(" -D__METAL_TYPED_BUFFER_READ_IMPL__=1");
 						Defines += TEXT(" -D__METAL_TYPED_BUFFER_RW_IMPL__=1");
 						break;
-					case EMetalTypeBufferModeTex:
+					case EMetalTypeBufferModeTB:
 						Defines += TEXT(" -D__METAL_TYPED_BUFFER_READ_IMPL__=3");
 						Defines += TEXT(" -D__METAL_TYPED_BUFFER_RW_IMPL__=3");
 						break;
-					case EMetalTypeBufferModeFun:
-						Defines += TEXT(" -D__METAL_TYPED_BUFFER_READ_IMPL__=2");
-						Defines += TEXT(" -D__METAL_TYPED_BUFFER_RW_IMPL__=2");
-						break;
 					default:
 						break;
 				}
@@ -1999,14 +1972,15 @@ void CompileShader_Metal(const FShaderCompilerInput& _Input,FShaderCompilerOutpu
 	bool bAppleTV = (Input.ShaderFormat == NAME_SF_METAL_TVOS || Input.ShaderFormat == NAME_SF_METAL_MRT_TVOS);
     if (Input.ShaderFormat == NAME_SF_METAL || Input.ShaderFormat == NAME_SF_METAL_TVOS)
 	{
-        VersionEnum = VersionEnum > 0 ? VersionEnum : 0;
+		UE_CLOG(VersionEnum < 2, LogShaders, Warning, TEXT("Metal shader version must be Metal v1.2 or higher for format %s!"), VersionEnum, *Input.ShaderFormat.ToString());
+        VersionEnum = VersionEnum >= 2 ? VersionEnum : 2;
         AdditionalDefines.SetDefine(TEXT("METAL_PROFILE"), 1);
 	}
 	else if (Input.ShaderFormat == NAME_SF_METAL_MRT || Input.ShaderFormat == NAME_SF_METAL_MRT_TVOS)
 	{
-        UE_CLOG(VersionEnum == 0, LogShaders, Warning, TEXT("Metal shader version should be Metal v1.2 or higher for format %s!"), VersionEnum, *Input.ShaderFormat.ToString());
+		UE_CLOG(VersionEnum < 2, LogShaders, Warning, TEXT("Metal shader version must be Metal v1.2 or higher for format %s!"), VersionEnum, *Input.ShaderFormat.ToString());
 		AdditionalDefines.SetDefine(TEXT("METAL_MRT_PROFILE"), 1);
-		VersionEnum = VersionEnum > 0 ? VersionEnum : 2;
+		VersionEnum = VersionEnum >= 2 ? VersionEnum : 2;
 		MetalCompilerTarget = HCT_FeatureLevelSM5;
 		Semantics = EMetalGPUSemanticsTBDRDesktop;
 	}
@@ -2066,7 +2040,7 @@ void CompileShader_Metal(const FShaderCompilerInput& _Input,FShaderCompilerOutpu
     {
         case 5:
             // Enable full SM5 feature support so tessellation & fragment UAVs compile
-            TypeMode = EMetalTypeBufferModeTex;
+            TypeMode = EMetalTypeBufferModeTB;
             HlslCompilerTarget = HCT_FeatureLevelSM5;
             StandardVersion = TEXT("2.1");
             if (bAppleTV)
@@ -2084,7 +2058,7 @@ void CompileShader_Metal(const FShaderCompilerInput& _Input,FShaderCompilerOutpu
             break;
 		case 4:
 			// Enable full SM5 feature support so tessellation & fragment UAVs compile
-			TypeMode = EMetalTypeBufferModeTex;
+			TypeMode = EMetalTypeBufferModeTB;
             HlslCompilerTarget = HCT_FeatureLevelSM5;
 			StandardVersion = TEXT("2.1");
 			if (bAppleTV)
@@ -2102,7 +2076,7 @@ void CompileShader_Metal(const FShaderCompilerInput& _Input,FShaderCompilerOutpu
 			break;
 		case 3:
 			// Enable full SM5 feature support so tessellation & fragment UAVs compile
-			TypeMode = EMetalTypeBufferModeUAV;
+			TypeMode = EMetalTypeBufferMode2D;
             HlslCompilerTarget = HCT_FeatureLevelSM5;
 			StandardVersion = TEXT("2.0");
 			if (bAppleTV)
@@ -2120,7 +2094,7 @@ void CompileShader_Metal(const FShaderCompilerInput& _Input,FShaderCompilerOutpu
 			break;
 		case 2:
 			// Enable full SM5 feature support so tessellation & fragment UAVs compile
-			TypeMode = EMetalTypeBufferModeSRV;
+			TypeMode = EMetalTypeBufferMode2D;
             HlslCompilerTarget = HCT_FeatureLevelSM5;
 			StandardVersion = TEXT("1.2");
 			if (bAppleTV)
diff --git a/Engine/Source/Developer/Apple/MetalShaderFormat/Private/MetalShaderFormat.cpp b/Engine/Source/Developer/Apple/MetalShaderFormat/Private/MetalShaderFormat.cpp
--- a/Engine/Source/Developer/Apple/MetalShaderFormat/Private/MetalShaderFormat.cpp
+++ b/Engine/Source/Developer/Apple/MetalShaderFormat/Private/MetalShaderFormat.cpp
@@ -239,7 +239,7 @@ class FMetalShaderFormat : public IShaderFormat
 public:
 	enum
 	{
-		HEADER_VERSION = 61,
+		HEADER_VERSION = 62,
 	};
 	
 	struct FVersion
diff --git a/Engine/Source/Developer/Apple/MetalShaderFormat/Private/MetalUtils.cpp b/Engine/Source/Developer/Apple/MetalShaderFormat/Private/MetalUtils.cpp
--- a/Engine/Source/Developer/Apple/MetalShaderFormat/Private/MetalUtils.cpp
+++ b/Engine/Source/Developer/Apple/MetalShaderFormat/Private/MetalUtils.cpp
@@ -203,7 +203,7 @@ namespace MetalUtils
 	{
 		{"SV_Depth", glsl_type::float_type, "FragDepth", ir_var_out, "[[ depth(any) ]]"},
 		{"SV_DepthLessEqual", glsl_type::float_type, "FragDepth", ir_var_out, "[[ depth(less) ]]"},
-		{"SV_Position", glsl_type::vec4_type, "IN_FragCoord", ir_var_in, "[[ position ]]"},
+		{"SV_Position", glsl_type::vec4_type, "IN_FragCoord", ir_var_in, "position"},
 		{"SV_IsFrontFace", glsl_type::bool_type, "IN_FrontFacing", ir_var_in, "[[ front_facing ]]"},
 		//{"SV_PrimitiveID", glsl_type::int_type, "IN_PrimitiveID", ir_var_in, "[[  ]]"},
 		//{"SV_RenderTargetArrayIndex", glsl_type::uint_type, "IN_Layer", ir_var_in, "[[ render_target_array_index ]]"},
@@ -248,7 +248,7 @@ namespace MetalUtils
 	{
 		{"SV_Depth", glsl_type::float_type, "FragDepth", ir_var_out, "[[ depth(any) ]]"},
 		{"SV_DepthLessEqual", glsl_type::float_type, "FragDepth", ir_var_out, "[[ depth(less) ]]"},
-		{"SV_Position", glsl_type::vec4_type, "IN_FragCoord", ir_var_in, "[[ position ]]"},
+		{"SV_Position", glsl_type::vec4_type, "IN_FragCoord", ir_var_in, "position"},
 		{"SV_IsFrontFace", glsl_type::bool_type, "IN_FrontFacing", ir_var_in, "[[ front_facing ]]"},
 		{"SV_Coverage", glsl_type::uint_type, "IN_Coverage", ir_var_in, "[[ sample_mask ]]"},
 		{"SV_Coverage", glsl_type::uint_type, "OUT_Coverage", ir_var_out, "[[ sample_mask ]]"},
@@ -392,9 +392,11 @@ namespace MetalUtils
 		DomainSystemValueTable,
 		ComputeSystemValueTable
 	};
+	
+	char const* interpolant_qualifiers[2][4] = {{"", "center_perspective", "flat", "center_no_perspective"}, {"centroid_perspective", "centroid_perspective", "flat", "centroid_no_perspective"}};
 
 	static ir_rvalue* GenerateInputFromSemantic(EHlslShaderFrequency Frequency, EMetalGPUSemantics bIsDesktop, _mesa_glsl_parse_state* ParseState,
-		const char* Semantic, const glsl_type* Type, char const* Name, exec_list* DeclInstructions, exec_list* PreCallInstructions)
+		const char* Semantic, FSemanticQualifier Qualifier, const glsl_type* Type, char const* Name, exec_list* DeclInstructions, exec_list* PreCallInstructions)
 	{
 		if (!Semantic)
 		{
@@ -431,7 +433,15 @@ namespace MetalUtils
 					if (!Variable)
 					{
 						Variable = new(ParseState) ir_variable(SystemValues[i].Type, SystemValues[i].MetalName, ir_var_in);
-						Variable->semantic = SystemValues[i].MetalSemantic;
+						if (!FCStringAnsi::Stricmp(Semantic, "SV_Position"))
+						{
+							char const* interp = Frequency == HSF_PixelShader ? interpolant_qualifiers[Qualifier.Fields.bCentroid][Qualifier.Fields.InterpolationMode] : "";
+							Variable->semantic = ralloc_asprintf(ParseState, "[[ %s, %s ]]", SystemValues[i].MetalSemantic, interp);
+						}
+						else
+						{
+							Variable->semantic = SystemValues[i].MetalSemantic;
+						}
 						Variable->read_only = true;
 						Variable->origin_upper_left = false;
 						DeclInstructions->push_tail(Variable);
@@ -531,9 +541,13 @@ namespace MetalUtils
 		else
 		if (!Variable->semantic)
 		{
-			Variable->semantic = ralloc_asprintf(ParseState, "[[ user(%s) ]]", Semantic);
+			char const* interp = Frequency == HSF_PixelShader ? interpolant_qualifiers[Qualifier.Fields.bCentroid][Qualifier.Fields.InterpolationMode] : "";
+			Variable->semantic = ralloc_asprintf(ParseState, "[[ user(%s), %s ]]", Semantic, interp);
 		}
 		Variable->read_only = true;
+		Variable->centroid = Qualifier.Fields.bCentroid;
+		Variable->interpolation = Qualifier.Fields.InterpolationMode;
+		Variable->is_patch_constant = Qualifier.Fields.bIsPatchConstant;
 		DeclInstructions->push_tail(Variable);
 		ParseState->symbols->add_variable(Variable);
 		ir_dereference_variable* VariableDeref = new(ParseState)ir_dereference_variable(Variable);
@@ -541,7 +555,7 @@ namespace MetalUtils
 	}
 
 	static void GenerateInputForVariable(EHlslShaderFrequency Frequency, EMetalGPUSemantics bIsDesktop, _mesa_glsl_parse_state* ParseState,
-		const char* InputSemantic, ir_dereference* InputVariableDeref, exec_list* DeclInstructions, exec_list* PreCallInstructions)
+		const char* InputSemantic, FSemanticQualifier Qualifier, ir_dereference* InputVariableDeref, exec_list* DeclInstructions, exec_list* PreCallInstructions)
 	{
 		const glsl_type* InputType = InputVariableDeref->type;
 		if (InputType->is_record())
@@ -577,7 +591,12 @@ namespace MetalUtils
 					ir_dereference_record* FieldDeref = new(ParseState)ir_dereference_record(
 						InputVariableDeref->clone(ParseState, NULL),
 						InputType->fields.structure[i].name);
-					GenerateInputForVariable(Frequency, bIsDesktop, ParseState, Semantic, FieldDeref, DeclInstructions, PreCallInstructions);
+					
+					FSemanticQualifier InQualifier = Qualifier;
+					InQualifier.Fields.bCentroid = InputType->fields.structure[i].centroid;
+					InQualifier.Fields.InterpolationMode = InputType->fields.structure[i].interpolation;
+					InQualifier.Fields.bIsPatchConstant = InputType->fields.structure[i].patchconstant;
+					GenerateInputForVariable(Frequency, bIsDesktop, ParseState, Semantic, InQualifier, FieldDeref, DeclInstructions, PreCallInstructions);
 				}
 				else
 				{
@@ -610,6 +629,7 @@ namespace MetalUtils
 						bIsDesktop,
 						ParseState,
 						ralloc_asprintf(ParseState, "%s%d", Semantic, BaseIndex + i),
+						Qualifier,
 						ArrayDeref,
 						DeclInstructions,
 						PreCallInstructions);
@@ -617,6 +637,7 @@ namespace MetalUtils
 			}
 			else
 			{
+				FSemanticQualifier InQualifier = Qualifier;
 				ir_dereference_variable* DerefVariable = InputVariableDeref->as_dereference_variable();
 				ir_dereference_array* DerefArray = InputVariableDeref->as_dereference_array();
 				ir_constant* DerefIndex = DerefArray ? DerefArray->array_index->as_constant() : nullptr;
@@ -629,13 +650,33 @@ namespace MetalUtils
 				else if (DerefIndex && DerefStruct)
 				{
 					Name = ralloc_asprintf(ParseState, "%s%d", DerefStruct->field, DerefIndex->value.u[0]);
+					
+					int FieldIndex = DerefStruct->type->field_index(DerefStruct->field);
+					if (FieldIndex >= 0)
+					{
+						DerefStruct->type->fields.structure[FieldIndex];
+						
+						InQualifier.Fields.bCentroid = DerefStruct->type->fields.structure[FieldIndex].centroid;
+						InQualifier.Fields.InterpolationMode = DerefStruct->type->fields.structure[FieldIndex].interpolation;
+						InQualifier.Fields.bIsPatchConstant = DerefStruct->type->fields.structure[FieldIndex].patchconstant;
+					}
 				}
 				else if(DerefStruct)
 				{
 					Name = DerefStruct->field;
+					
+					int FieldIndex = DerefStruct->type->field_index(DerefStruct->field);
+					if (FieldIndex >= 0)
+					{
+						DerefStruct->type->fields.structure[FieldIndex];
+						
+						InQualifier.Fields.bCentroid = DerefStruct->type->fields.structure[FieldIndex].centroid;
+						InQualifier.Fields.InterpolationMode = DerefStruct->type->fields.structure[FieldIndex].interpolation;
+						InQualifier.Fields.bIsPatchConstant = DerefStruct->type->fields.structure[FieldIndex].patchconstant;
+					}
 				}
 				
-				ir_rvalue* SrcValue = GenerateInputFromSemantic(Frequency, bIsDesktop, ParseState, InputSemantic, InputType, Name, DeclInstructions, PreCallInstructions);
+				ir_rvalue* SrcValue = GenerateInputFromSemantic(Frequency, bIsDesktop, ParseState, InputSemantic, InQualifier, InputType, Name, DeclInstructions, PreCallInstructions);
 				if (SrcValue)
 				{
 					YYLTYPE loc;
@@ -647,16 +688,16 @@ namespace MetalUtils
 		}
 	}
 
-	ir_dereference_variable* GenerateInput(EHlslShaderFrequency Frequency, uint32 bIsDesktop, _mesa_glsl_parse_state* ParseState, const char* InputName, const char* InputSemantic, const glsl_type* InputType, exec_list* DeclInstructions, exec_list* PreCallInstructions)
+	ir_dereference_variable* GenerateInput(EHlslShaderFrequency Frequency, uint32 bIsDesktop, _mesa_glsl_parse_state* ParseState, const char* InputName, const char* InputSemantic, FSemanticQualifier Qualifier, const glsl_type* InputType, exec_list* DeclInstructions, exec_list* PreCallInstructions)
 	{
 		if((InputType->is_inputpatch()))
 		{
-			return GenerateInputFromSemantic(Frequency, (EMetalGPUSemantics)bIsDesktop, ParseState, InputSemantic, InputType, nullptr, DeclInstructions, PreCallInstructions)->as_dereference_variable();
+			return GenerateInputFromSemantic(Frequency, (EMetalGPUSemantics)bIsDesktop, ParseState, InputSemantic, Qualifier, InputType, nullptr, DeclInstructions, PreCallInstructions)->as_dereference_variable();
 		}
 		ir_variable* TempVariable = new(ParseState)ir_variable(InputType, InputName ? ralloc_strdup(ParseState, InputName) : nullptr, ir_var_temporary);
 		ir_dereference_variable* TempVariableDeref = new(ParseState)ir_dereference_variable(TempVariable);
 		PreCallInstructions->push_tail(TempVariable);
-		GenerateInputForVariable(Frequency, (EMetalGPUSemantics)bIsDesktop, ParseState, InputSemantic, TempVariableDeref, DeclInstructions, PreCallInstructions);
+		GenerateInputForVariable(Frequency, (EMetalGPUSemantics)bIsDesktop, ParseState, InputSemantic, Qualifier, TempVariableDeref, DeclInstructions, PreCallInstructions);
 		return TempVariableDeref;
 	}
 
@@ -697,42 +738,9 @@ namespace MetalUtils
 		// As it happens we already order them by importance (0: Global > 1: VR-instanced fallback > 2: vertex-shader-layer)
 		uint32 const ClipPrefixLen = 15;
 		FMetalLanguageSpec* Spec = (FMetalLanguageSpec*)ParseState->LanguageSpec;
-		if ((bIsDesktop == EMetalGPUSemanticsImmediateDesktop && Spec->Version < 3) && Semantic && FCStringAnsi::Strnicmp(Semantic, "SV_ClipDistance", ClipPrefixLen) == 0 && !Variable)
-		{
-			uint32 const Count = Spec->GetClipDistanceCount();
-			uint32 const Used = Spec->ClipDistancesUsed;
-			check(Count > 0);
-			
-			uint32 Index = 0;
-			if (Semantic[ClipPrefixLen] >= '1' && Semantic[ClipPrefixLen] <= '7')
-			{
-				Index = Semantic[ClipPrefixLen] - '0';
-			}
-			
-			bool bUsed = false;
-			for (uint32 i = 0; !bUsed && i < Index; i++)
-			{
-				bUsed = (Used & (1 << i)) != 0;
-			}
-			
-			if (!bUsed)
-			{
-				ir_variable* compacted_clip = ParseState->symbols->get_variable("clip_distance_array");
-				if (!compacted_clip)
-				{
-					compacted_clip = new(ParseState)ir_variable(glsl_type::float_type, "clip_distance_array", ir_var_out);
-					compacted_clip->semantic = ralloc_asprintf(ParseState, "[[ clip_distance ]]");
-					DeclInstructions->push_tail(compacted_clip);
-					ParseState->symbols->add_variable(compacted_clip);
-				}
-				*DestVariableType = glsl_type::float_type;
-				ir_rvalue* VariableDeref = new(ParseState)ir_dereference_variable(compacted_clip);
-				return VariableDeref;
-			}
-		}
 		
 		// But for iOS/tvOS and future, non-broken desktop we can just remap the variable to the actual clip-distance-array
-		if ((bIsDesktop != EMetalGPUSemanticsImmediateDesktop || Spec->Version >= 3) && Semantic && FCStringAnsi::Strnicmp(Semantic, "SV_ClipDistance", ClipPrefixLen) == 0 && !Variable)
+		if (Semantic && FCStringAnsi::Strnicmp(Semantic, "SV_ClipDistance", ClipPrefixLen) == 0 && !Variable)
 		{
 			Variable = ParseState->symbols->get_variable("clip_distance_array");
 			
@@ -1229,18 +1237,12 @@ void FMetalCodeBackend::MovePackedUniformsToMain(exec_list* ir, _mesa_glsl_parse
 			switch(TypedMode)
 			{
 				case EMetalTypeBufferModeRaw:
-                case EMetalTypeBufferModeFun:
 				{
 					bIsBuffer = (!Var->type->is_sampler() && !Var->type->is_image()) || Var->type->sampler_buffer;
 					break;
 				}
-				case EMetalTypeBufferModeSRV:
-				{
-					bIsBuffer = (!Var->type->is_sampler() && !Var->type->is_image()) || (Var->type->sampler_buffer && (Var->type->is_image() || bIsStructuredBuffer || bIsByteAddressBuffer || bIsInvariant || OutBuffers.AtomicVariables.find(Var) != OutBuffers.AtomicVariables.end())) || bIsVec3;
-					break;
-				}
-				case EMetalTypeBufferModeUAV:
-				case EMetalTypeBufferModeTex:
+				case EMetalTypeBufferMode2D:
+				case EMetalTypeBufferModeTB:
 				{
 					bIsBuffer = (!Var->type->is_sampler() && !Var->type->is_image()) || (Var->type->sampler_buffer && (OutBuffers.AtomicVariables.find(Var) != OutBuffers.AtomicVariables.end() || bIsStructuredBuffer || bIsInvariant || bIsByteAddressBuffer)) || bIsVec3;
 					break;
@@ -1395,6 +1397,10 @@ static bool ProcessStageInVariables(_mesa_glsl_parse_state* ParseState, EMetalGP
 						OutMember.type = Variable->type->element_type();
 						OutMember.semantic = ralloc_asprintf(ParseState, "ATTRIBUTE%d", AttributeIndex);
 						OutMember.name = ralloc_asprintf(ParseState, "ATTRIBUTE%d_%s", AttributeIndex, Variable->name);
+						OutMember.centroid = Variable->centroid;
+						OutMember.interpolation = Variable->interpolation;
+						OutMember.geometryinput = Variable->geometryinput;
+						OutMember.patchconstant = Variable->is_patch_constant;
 
 						if (OutVertexAttributesMask)
 						{
@@ -1410,6 +1416,10 @@ static bool ProcessStageInVariables(_mesa_glsl_parse_state* ParseState, EMetalGP
 					OutMember.type = Variable->type;
 					OutMember.semantic = ralloc_asprintf(ParseState, "ATTRIBUTE%d", AttributeIndex);
 					OutMember.name = Variable->name;
+					OutMember.centroid = Variable->centroid;
+					OutMember.interpolation = Variable->interpolation;
+					OutMember.geometryinput = Variable->geometryinput;
+					OutMember.patchconstant = Variable->is_patch_constant;
 
 					if (OutVertexAttributesMask)
 					{
@@ -1450,6 +1460,10 @@ static bool ProcessStageInVariables(_mesa_glsl_parse_state* ParseState, EMetalGP
 	Member.type = Variable->type;
 	Member.name = ralloc_strdup(ParseState, Variable->name);
 	Member.semantic = ralloc_strdup(ParseState, Variable->semantic ? Variable->semantic : Variable->name);
+	Member.centroid = Variable->centroid;
+	Member.interpolation = Variable->interpolation;
+	Member.geometryinput = Variable->geometryinput;
+	Member.patchconstant = Variable->is_patch_constant;
 	OutStageInMembers.Add(Member);
 	OutStageInVariables.insert(Variable);
 
@@ -1522,6 +1536,7 @@ static ir_dereference_variable* GenerateShaderInput(
 	EHlslShaderFrequency Frequency, EMetalGPUSemantics bIsDesktop,
 	_mesa_glsl_parse_state* ParseState,
 	const char* InputSemantic,
+	FSemanticQualifier Qualifier,
 	const glsl_type* InputType,
 	exec_list* DeclInstructions,
 	exec_list* PreCallInstructions)
@@ -1534,7 +1549,7 @@ static ir_dereference_variable* GenerateShaderInput(
 	PreCallInstructions->push_tail(TempVariable);
 
 	check(!InputType->is_inputpatch() && !InputType->is_outputpatch());
-	ir_rvalue* SrcValue = MetalUtils::GenerateInputFromSemantic(Frequency, bIsDesktop, ParseState, InputSemantic, InputType, nullptr, DeclInstructions, PreCallInstructions);
+	ir_rvalue* SrcValue = MetalUtils::GenerateInputFromSemantic(Frequency, bIsDesktop, ParseState, InputSemantic, Qualifier, InputType, nullptr, DeclInstructions, PreCallInstructions);
 	if(SrcValue)
 	{
 		YYLTYPE loc ={0};
@@ -2207,6 +2222,12 @@ void FMetalCodeBackend::PackInputsAndOutputs(exec_list* Instructions, _mesa_glsl
 						Member.type = Variable->type;
 						Member.name = ralloc_strdup(ParseState, Variable->name);
 						Member.semantic = ralloc_strdup(ParseState, Variable->semantic ? Variable->semantic : Variable->name);
+						
+						Member.centroid = Variable->centroid;
+						Member.interpolation = Variable->interpolation;
+						Member.geometryinput = Variable->geometryinput;
+						Member.patchconstant = Variable->is_patch_constant;
+						
 						VSOutMembers.Add(Member);
 						VSOutVariables.insert(Variable);
 					}
@@ -2333,6 +2354,10 @@ void FMetalCodeBackend::PackInputsAndOutputs(exec_list* Instructions, _mesa_glsl
 						Member.type = Variable->type;
 						Member.name = ralloc_strdup(ParseState, Variable->name);
 						Member.semantic = ralloc_strdup(ParseState, Variable->semantic ? Variable->semantic : Variable->name);
+						Member.centroid = Variable->centroid;
+						Member.interpolation = Variable->interpolation;
+						Member.geometryinput = Variable->geometryinput;
+						Member.patchconstant = Variable->is_patch_constant;
 						PSOutMembers.Add(Member);
 						PSOutVariables.insert(Variable);
 					}
@@ -2456,6 +2481,10 @@ void FMetalCodeBackend::PackInputsAndOutputs(exec_list* Instructions, _mesa_glsl
 						Member.type = Variable->type;
 						Member.name = ralloc_strdup(ParseState, Variable->name);
 						Member.semantic = ralloc_strdup(ParseState, Variable->semantic ? Variable->semantic : Variable->name);
+						Member.centroid = Variable->centroid;
+						Member.interpolation = Variable->interpolation;
+						Member.geometryinput = Variable->geometryinput;
+						Member.patchconstant = Variable->is_patch_constant;
 						DSOutMembers.Add(Member);
 						DSOutVariables.insert(Variable);
 					}
@@ -2470,6 +2499,10 @@ void FMetalCodeBackend::PackInputsAndOutputs(exec_list* Instructions, _mesa_glsl
 						Member.type = Variable->type;
 						Member.name = ralloc_strdup(ParseState, Variable->name);
 						Member.semantic = ralloc_strdup(ParseState, Variable->semantic ? Variable->semantic : Variable->name);
+						Member.centroid = Variable->centroid;
+						Member.interpolation = Variable->interpolation;
+						Member.geometryinput = Variable->geometryinput;
+						Member.patchconstant = Variable->is_patch_constant;
 						DSPatchMembers.Add(Member);
 						DSPatchVariables.insert(Variable);
 					}
@@ -2590,9 +2623,13 @@ void FMetalCodeBackend::PackInputsAndOutputs(exec_list* Instructions, _mesa_glsl
 					exec_list PreMainDomainCallInstructions;
 
 					FSemanticQualifier Qualifier;
+					Qualifier.Fields.bCentroid = Variable->centroid;
+					Qualifier.Fields.InterpolationMode = Variable->interpolation;
+					Qualifier.Fields.bIsPatchConstant = Variable->is_patch_constant;
+					Qualifier.Fields.bIsTessellationVSHS = false;
 
 					// @todo there has to be a better way to handle this vs looping over GenerateInput...
-					auto deref = MetalUtils::GenerateInput(Frequency, bIsDesktop, ParseState, Variable->name, Variable->semantic, Variable->type->inner_type, &MainDomainDeclInstructions, &PreMainDomainCallInstructions);
+					auto deref = MetalUtils::GenerateInput(Frequency, bIsDesktop, ParseState, Variable->name, Variable->semantic, Qualifier, Variable->type->inner_type, &MainDomainDeclInstructions, &PreMainDomainCallInstructions);
 
 					// make a flat perControlPoint struct
 					ir_dereference_variable* OutputControlPointDeref = NULL;
@@ -2617,6 +2654,10 @@ void FMetalCodeBackend::PackInputsAndOutputs(exec_list* Instructions, _mesa_glsl
 										InnerVariable->name = ralloc_asprintf(ParseState, "OUT_ATTRIBUTE%d_%s", InnerAttribute, InnerVariable->name);
 										Member.name = ralloc_strdup(ParseState, InnerVariable->name);
 										Member.semantic = ralloc_asprintf(ParseState, "[[ attribute(%d) ]]", InnerAttribute);
+										Member.centroid = InnerVariable->centroid;
+										Member.interpolation = InnerVariable->interpolation;
+										Member.geometryinput = InnerVariable->geometryinput;
+										Member.patchconstant = InnerVariable->is_patch_constant;
 
 										if (OutputVertex == 0)
 										{
@@ -2817,12 +2858,14 @@ void FMetalCodeBackend::PackInputsAndOutputs(exec_list* Instructions, _mesa_glsl
 #endif // USE_DS_ATTRIBUTES
 				else
 				{
+					FSemanticQualifier Qualifier;
 					// At this point this should be a built-in system value
 					check(Variable->semantic);
 					ArgVarDeref = GenerateShaderInput(
 						Frequency, bIsDesktop,
 						ParseState,
 						Variable->semantic,
+						Qualifier,
 						Variable->type,
 						&DeclInstructions,
 						&PreCallInstructions
diff --git a/Engine/Source/Developer/Apple/MetalShaderFormat/Private/MetalUtils.h b/Engine/Source/Developer/Apple/MetalShaderFormat/Private/MetalUtils.h
--- a/Engine/Source/Developer/Apple/MetalShaderFormat/Private/MetalUtils.h
+++ b/Engine/Source/Developer/Apple/MetalShaderFormat/Private/MetalUtils.h
@@ -410,18 +410,23 @@ struct FBuffers
 
 const glsl_type* PromoteHalfToFloatType(_mesa_glsl_parse_state* state, const glsl_type* type);
 
-struct FSemanticQualifier
+union FSemanticQualifier
 {
 	struct
 	{
-		bool bIsPatchConstant = false;
-		bool bIsTessellationVSHS = false; // TODO FIXME UGLY UGLY HACK (this is not the right place to put this flag)
+		unsigned bCentroid : 1;
+		unsigned InterpolationMode : 2;
+		unsigned bIsPatchConstant : 1;
+		unsigned bIsTessellationVSHS : 1; // TODO FIXME UGLY UGLY HACK (this is not the right place to put this flag)
 	} Fields;
+	unsigned Packed;
+	
+	FSemanticQualifier() : Packed(0) {}
 };
 
 namespace MetalUtils
 {
-	ir_dereference_variable* GenerateInput(EHlslShaderFrequency Frequency, uint32 bIsDesktop, _mesa_glsl_parse_state* ParseState, const char* InputName, const char* InputSemantic,
+	ir_dereference_variable* GenerateInput(EHlslShaderFrequency Frequency, uint32 bIsDesktop, _mesa_glsl_parse_state* ParseState, const char* InputName, const char* InputSemantic, FSemanticQualifier Qualifier,
 		const glsl_type* InputType, exec_list* DeclInstructions, exec_list* PreCallInstructions);
 
 	ir_dereference_variable* GenerateOutput(EHlslShaderFrequency Frequency, uint32 bIsDesktop, _mesa_glsl_parse_state* ParseState, const char* OutputSemantic,
diff --git a/Engine/Source/Developer/IOS/IOSPlatformEditor/Private/IOSTargetSettingsCustomization.cpp b/Engine/Source/Developer/IOS/IOSPlatformEditor/Private/IOSTargetSettingsCustomization.cpp
--- a/Engine/Source/Developer/IOS/IOSPlatformEditor/Private/IOSTargetSettingsCustomization.cpp
+++ b/Engine/Source/Developer/IOS/IOSPlatformEditor/Private/IOSTargetSettingsCustomization.cpp
@@ -1842,46 +1842,6 @@ void FIOSTargetSettingsCustomization::UpdateShaderStandardWarning()
 
 void FIOSTargetSettingsCustomization::UpdateOSVersionWarning()
 {
-	bool bMRTEnabled = false;
-	if (MRTPropertyHandle.IsValid())
-	{
-		MRTPropertyHandle->GetValue(bMRTEnabled);
-	}
-	
-	// Due to a driver bug on A8 devices running iOS 9 we can only support the global clip-plane when running iOS 10+
-	static IConsoleVariable* ClipPlaneCVar = IConsoleManager::Get().FindConsoleVariable(TEXT("r.AllowGlobalClipPlane"));
-	if (MinOSPropertyHandle.IsValid() && IOSVersionWarningTextBox.IsValid() && ((ClipPlaneCVar && ClipPlaneCVar->GetInt() != 0) || bMRTEnabled))
-	{
-		uint8 EnumValue;
-		MinOSPropertyHandle->GetValue(EnumValue);
-		
-		if (EnumValue < (uint8)EIOSVersion::IOS_10)
-		{
-			SetMinVersion((int32)EIOSVersion::IOS_10);
-		}
-		
-		FText Message;
-		if (bMRTEnabled)
-		{
-			Message = LOCTEXT("MetalMRTStandardv1.2","Enabling the Desktop Forward Renderer Metal requires Shader Standard v1.2 which increases the minimum operating system requirement for Metal from iOS 8.0 or later to iOS 10.0 or later.");
-			
-			if (ShaderVersionPropertyHandle.IsValid())
-			{
-				ShaderVersionPropertyHandle->GetValue(EnumValue);
-				if (EnumValue < (uint8)EIOSMetalShaderStandard::IOSMetalSLStandard_1_2)
-				{
-					SetShaderStandard((int32)EIOSMetalShaderStandard::IOSMetalSLStandard_1_2);
-				}
-			}
-		}
-		else
-		{
-			Message = LOCTEXT("GlobalClipPlaneiOS10", "Enabling the Global Clip Plane increases the minimum operating system requirement for Metal from iOS 8.0 or later to iOS 10.0 or later.");
-		}
-		
-		// Update the UI
-		IOSVersionWarningTextBox->SetError(Message);
-	}
 }
 
 void FIOSTargetSettingsCustomization::UpdateMetalMRTWarning()
diff --git a/Engine/Source/Developer/Mac/MacTargetPlatform/Private/MacTargetPlatformModule.cpp b/Engine/Source/Developer/Mac/MacTargetPlatform/Private/MacTargetPlatformModule.cpp
--- a/Engine/Source/Developer/Mac/MacTargetPlatform/Private/MacTargetPlatformModule.cpp
+++ b/Engine/Source/Developer/Mac/MacTargetPlatform/Private/MacTargetPlatformModule.cpp
@@ -66,7 +66,7 @@ public:
 		
 		int32 Value = 1;
         GConfig->GetInt(TEXT("/Script/MacTargetPlatform.MacTargetSettings"), TEXT("MaxShaderLanguageVersion"), Value, GEngineIni);
-        TargetSettings->MaxShaderLanguageVersion = FMath::Max(Value, 2);
+        TargetSettings->MaxShaderLanguageVersion = FMath::Max(Value, 3);
 		
 		if (!GConfig->GetBool(TEXT("/Script/MacTargetPlatform.MacTargetSettings"), TEXT("UseFastIntrinsics"), TargetSettings->UseFastIntrinsics, GEngineIni))
 		{
diff --git a/Engine/Source/Runtime/Apple/MetalRHI/Private/MetalCommandQueue.cpp b/Engine/Source/Runtime/Apple/MetalRHI/Private/MetalCommandQueue.cpp
--- a/Engine/Source/Runtime/Apple/MetalRHI/Private/MetalCommandQueue.cpp
+++ b/Engine/Source/Runtime/Apple/MetalRHI/Private/MetalCommandQueue.cpp
@@ -18,6 +18,7 @@
 
 #pragma mark - Private C++ Statics -
 uint64 FMetalCommandQueue::Features = 0;
+extern mtlpp::VertexFormat GMetalFColorVertexFormat;
 
 #pragma mark - Public C++ Boilerplate -
 
@@ -84,6 +85,8 @@ FMetalCommandQueue::FMetalCommandQueue(mtlpp::Device InDevice, uint32 const MaxN
 			{
 				Features |= EMetalFeaturesGPUCaptureManager | EMetalFeaturesBufferSubAllocation | EMetalFeaturesParallelRenderEncoders | EMetalFeaturesPipelineBufferMutability;
 				
+				GMetalFColorVertexFormat = mtlpp::VertexFormat::UChar4Normalized_BGRA;
+				
 				if (MaxShaderVersion >= 3)
 				{
 					Features |= EMetalFeaturesLinearTextureUAVs;
@@ -141,34 +144,36 @@ FMetalCommandQueue::FMetalCommandQueue(mtlpp::Device InDevice, uint32 const MaxN
 				Features |= EMetalFeaturesDeferredStoreActions | EMetalFeaturesCombinedDepthStencil;
 			}
 			
-			// Turn on Linear Texture UAVs! Avoids the need to have function-constants which reduces initial runtime shader compile time
-			if (MaxShaderVersion >= 3 && Vers.majorVersion >= 11)
-			{
-				Features |= EMetalFeaturesLinearTextureUAVs;
-			}
-			
-			// Turn on Texture Buffers! These are faster on the GPU as we don't need to do out-of-bounds tests but require Metal 2.1 and macOS 10.14
-			if (Vers.majorVersion >= 12)
+			if(Vers.majorVersion >= 11)
 			{
-				Features |= EMetalFeaturesMaxThreadsPerThreadgroup;
-				Features |= EMetalFeaturesFences;
-				Features |= EMetalFeaturesHeaps;
+				GMetalFColorVertexFormat = mtlpp::VertexFormat::UChar4Normalized_BGRA;
 				
-				if (MaxShaderVersion >= 4)
+				Features |= EMetalFeaturesPresentMinDuration | EMetalFeaturesGPUCaptureManager | EMetalFeaturesBufferSubAllocation | EMetalFeaturesParallelRenderEncoders | EMetalFeaturesPipelineBufferMutability;
+				
+				// Turn on Linear Texture UAVs! Avoids the need to have function-constants which reduces initial runtime shader compile time
+				if (MaxShaderVersion >= 3)
 				{
-					Features |= EMetalFeaturesTextureBuffers;
+					Features |= EMetalFeaturesLinearTextureUAVs;
 				}
-				if (Device.SupportsFeatureSet(mtlpp::FeatureSet::iOS_GPUFamily5_v1))
+				
+				// Turn on Texture Buffers! These are faster on the GPU as we don't need to do out-of-bounds tests but require Metal 2.1 and macOS 10.14
+				if (Vers.majorVersion >= 12)
 				{
-					Features |= EMetalFeaturesLayeredRendering;
+					Features |= EMetalFeaturesMaxThreadsPerThreadgroup;
+					Features |= EMetalFeaturesFences;
+					Features |= EMetalFeaturesHeaps;
+					
+					if (MaxShaderVersion >= 4)
+					{
+						Features |= EMetalFeaturesTextureBuffers;
+					}
+					if (Device.SupportsFeatureSet(mtlpp::FeatureSet::iOS_GPUFamily5_v1))
+					{
+						Features |= EMetalFeaturesLayeredRendering;
+					}
 				}
 			}
         }
-		
-		if(Vers.majorVersion >= 11)
-		{
-			Features |= EMetalFeaturesPresentMinDuration | EMetalFeaturesGPUCaptureManager | EMetalFeaturesBufferSubAllocation | EMetalFeaturesParallelRenderEncoders | EMetalFeaturesPipelineBufferMutability;
-        }
 #endif
 	}
 	else if(Vers.majorVersion == 8 && Vers.minorVersion >= 3)
@@ -207,6 +212,8 @@ FMetalCommandQueue::FMetalCommandQueue(mtlpp::Device InDevice, uint32 const MaxN
 			}
 		}
 		
+		GMetalFColorVertexFormat = mtlpp::VertexFormat::UChar4Normalized_BGRA;
+		
 		// On 10.13.5+ we can use MTLParallelRenderEncoder
 		if (FPlatformMisc::MacOSXVersionCompare(10,13,5) >= 0)
 		{
diff --git a/Engine/Source/Runtime/Apple/MetalRHI/Private/MetalPipeline.cpp b/Engine/Source/Runtime/Apple/MetalRHI/Private/MetalPipeline.cpp
--- a/Engine/Source/Runtime/Apple/MetalRHI/Private/MetalPipeline.cpp
+++ b/Engine/Source/Runtime/Apple/MetalRHI/Private/MetalPipeline.cpp
@@ -162,9 +162,6 @@ struct FMetalGraphicsPipelineKey
 	FSHAHash VertexFunction;
 	FSHAHash DomainFunction;
 	FSHAHash PixelFunction;
-	uint32 VertexBufferHash;
-	uint32 DomainBufferHash;
-	uint32 PixelBufferHash;
 
 	template<typename Type>
 	inline void SetHashValue(uint32 Offset, uint32 NumBits, Type Value)
@@ -188,10 +185,7 @@ struct FMetalGraphicsPipelineKey
 		&& VertexDescriptorHash == Other.VertexDescriptorHash
 		&& VertexFunction == Other.VertexFunction
 		&& DomainFunction == Other.DomainFunction
-		&& PixelFunction == Other.PixelFunction
-		&& VertexBufferHash == Other.VertexBufferHash
-		&& DomainBufferHash == Other.DomainBufferHash
-		&& PixelBufferHash == Other.PixelBufferHash);
+		&& PixelFunction == Other.PixelFunction);
 	}
 	
 	friend uint32 GetTypeHash(FMetalGraphicsPipelineKey const& Key)
@@ -200,13 +194,10 @@ struct FMetalGraphicsPipelineKey
 		H = FCrc::MemCrc32(Key.VertexFunction.Hash, sizeof(Key.VertexFunction.Hash), H);
 		H = FCrc::MemCrc32(Key.DomainFunction.Hash, sizeof(Key.DomainFunction.Hash), H);
 		H = FCrc::MemCrc32(Key.PixelFunction.Hash, sizeof(Key.PixelFunction.Hash), H);
-		H = HashCombine(H, GetTypeHash(Key.VertexBufferHash));
-		H = HashCombine(H, GetTypeHash(Key.DomainBufferHash));
-		H = HashCombine(H, GetTypeHash(Key.PixelBufferHash));
 		return H;
 	}
 	
-	friend void InitMetalGraphicsPipelineKey(FMetalGraphicsPipelineKey& Key, const FGraphicsPipelineStateInitializer& Init, EMetalIndexType const IndexType, EPixelFormat const* const VertexBufferTypes, EPixelFormat const* const PixelBufferTypes, EPixelFormat const* const DomainBufferTypes)
+	friend void InitMetalGraphicsPipelineKey(FMetalGraphicsPipelineKey& Key, const FGraphicsPipelineStateInitializer& Init, EMetalIndexType const IndexType)
 	{
 		uint32 const NumActiveTargets = Init.ComputeNumValidRenderTargets();
 		check(NumActiveTargets <= MaxSimultaneousRenderTargets);
@@ -296,26 +287,18 @@ struct FMetalGraphicsPipelineKey
 		FMetalDomainShader* DomainShader = (FMetalDomainShader*)Init.BoundShaderState.DomainShaderRHI;
 		
         Key.VertexFunction = VertexShader->GetHash();
-		Key.VertexBufferHash = VertexShader->GetBindingHash(VertexBufferTypes);
 		if (DomainShader)
 		{
 			Key.DomainFunction = DomainShader->GetHash();
 			Key.SetHashValue(Offset_IndexType, NumBits_IndexType, IndexType);
-			Key.DomainBufferHash = DomainShader->GetBindingHash(DomainBufferTypes);
 		}
 		else
 		{
 			Key.SetHashValue(Offset_IndexType, NumBits_IndexType, EMetalIndexType_None);
-			Key.DomainBufferHash = 0;
 		}
 		if (PixelShader)
 		{
 			Key.PixelFunction = PixelShader->GetHash();
-			Key.PixelBufferHash = PixelShader->GetBindingHash(PixelBufferTypes);
-		}
-		else
-		{
-			Key.PixelBufferHash = 0;
 		}
 	}
 };
@@ -353,15 +336,15 @@ static MTLVertexDescriptor* GetMaskedVertexDescriptor(MTLVertexDescriptor* Input
 	return InputDesc;
 }
 
-static FMetalShaderPipeline* CreateMTLRenderPipeline(bool const bSync, FMetalGraphicsPipelineKey const& Key, const FGraphicsPipelineStateInitializer& Init, EMetalIndexType const IndexType, EPixelFormat const* const VertexBufferTypes, EPixelFormat const* const PixelBufferTypes, EPixelFormat const* const DomainBufferTypes)
+static FMetalShaderPipeline* CreateMTLRenderPipeline(bool const bSync, FMetalGraphicsPipelineKey const& Key, const FGraphicsPipelineStateInitializer& Init, EMetalIndexType const IndexType)
 {
     FMetalVertexShader* VertexShader = (FMetalVertexShader*)Init.BoundShaderState.VertexShaderRHI;
     FMetalDomainShader* DomainShader = (FMetalDomainShader*)Init.BoundShaderState.DomainShaderRHI;
     FMetalPixelShader* PixelShader = (FMetalPixelShader*)Init.BoundShaderState.PixelShaderRHI;
     
-    mtlpp::Function vertexFunction = VertexShader->GetFunction(IndexType, VertexBufferTypes, Key.VertexBufferHash);
-    mtlpp::Function fragmentFunction = PixelShader ? PixelShader->GetFunction(EMetalIndexType_None, PixelBufferTypes, Key.PixelBufferHash) : nil;
-    mtlpp::Function domainFunction = DomainShader ? DomainShader->GetFunction(EMetalIndexType_None, DomainBufferTypes, Key.DomainBufferHash) : nil;
+    mtlpp::Function vertexFunction = VertexShader->GetFunction(IndexType);
+    mtlpp::Function fragmentFunction = PixelShader ? PixelShader->GetFunction(EMetalIndexType_None) : nil;
+    mtlpp::Function domainFunction = DomainShader ? DomainShader->GetFunction(EMetalIndexType_None) : nil;
     
     FMetalShaderPipeline* Pipeline = nil;
     if (vertexFunction && ((PixelShader != nullptr) == (fragmentFunction != nil)) && ((DomainShader != nullptr) == (domainFunction != nil)))
@@ -945,7 +928,7 @@ static FMetalShaderPipeline* CreateMTLRenderPipeline(bool const bSync, FMetalGra
     return !bSync ? nil : Pipeline;
 }
 
-static FMetalShaderPipeline* GetMTLRenderPipeline(bool const bSync, FMetalGraphicsPipelineState const* State, const FGraphicsPipelineStateInitializer& Init, EMetalIndexType const IndexType, EPixelFormat const* const VertexBufferTypes, EPixelFormat const* const PixelBufferTypes, EPixelFormat const* const DomainBufferTypes)
+static FMetalShaderPipeline* GetMTLRenderPipeline(bool const bSync, FMetalGraphicsPipelineState const* State, const FGraphicsPipelineStateInitializer& Init, EMetalIndexType const IndexType)
 {
 	static FRWLock PipelineMutex;
 	static TMap<FMetalGraphicsPipelineKey, FMetalShaderPipeline*> Pipelines;
@@ -953,7 +936,7 @@ static FMetalShaderPipeline* GetMTLRenderPipeline(bool const bSync, FMetalGraphi
 	SCOPE_CYCLE_COUNTER(STAT_MetalPipelineStateTime);
 	
 	FMetalGraphicsPipelineKey Key;
-	InitMetalGraphicsPipelineKey(Key, Init, IndexType, VertexBufferTypes, PixelBufferTypes, DomainBufferTypes);
+	InitMetalGraphicsPipelineKey(Key, Init, IndexType);
 
 	// By default there'll be more threads trying to read this than to write it.
 	FRWScopeLock Lock(PipelineMutex, SLT_ReadOnly);
@@ -962,7 +945,7 @@ static FMetalShaderPipeline* GetMTLRenderPipeline(bool const bSync, FMetalGraphi
 	FMetalShaderPipeline* Desc = Pipelines.FindRef(Key);
 	if (Desc == nil)
 	{
-		Desc = CreateMTLRenderPipeline(bSync, Key, Init, IndexType, VertexBufferTypes, PixelBufferTypes, DomainBufferTypes);
+		Desc = CreateMTLRenderPipeline(bSync, Key, Init, IndexType);
 		
 		// Bail cleanly if compilation fails.
 		if(!Desc)
@@ -989,8 +972,8 @@ bool FMetalGraphicsPipelineState::Compile()
 	FMemory::Memzero(PipelineStates);
 	for (uint32 i = 0; i < EMetalIndexType_Num; i++)
 	{
-		PipelineStates[i][0][0][0] = [GetMTLRenderPipeline(true, this, Initializer, (EMetalIndexType)i, nullptr, nullptr, nullptr) retain];
-		if(!PipelineStates[i][0][0][0])
+		PipelineStates[i] = [GetMTLRenderPipeline(true, this, Initializer, (EMetalIndexType)i) retain];
+		if(!PipelineStates[i])
 		{
 			return false;
 		}
@@ -1001,47 +984,23 @@ bool FMetalGraphicsPipelineState::Compile()
 
 FMetalGraphicsPipelineState::~FMetalGraphicsPipelineState()
 {
-    static uint32 MaxBufferNum = (GMaxRHIFeatureLevel == ERHIFeatureLevel::SM5) ? EMetalBufferType_Num : 1u;
 	for (uint32 i = 0; i < EMetalIndexType_Num; i++)
 	{
-        for (uint32 v = 0; v < MaxBufferNum; v++)
-        {
-            for (uint32 f = 0; f < MaxBufferNum; f++)
-            {
-                for (uint32 c = 0; c < MaxBufferNum; c++)
-                {
-                	if (PipelineStates[i][v][f][c])
-                	{
-	                    [PipelineStates[i][v][f][c] release];
-	                    PipelineStates[i][v][f][c] = nil;
-                    }
-                }
-            }
-        }
+		[PipelineStates[i] release];
+		PipelineStates[i] = nil;
 	}
 }
 
-FMetalShaderPipeline* FMetalGraphicsPipelineState::GetPipeline(EMetalIndexType IndexType, uint32 VertexBufferHash, uint32 PixelBufferHash, uint32 DomainBufferHash, EPixelFormat const* const VertexBufferTypes, EPixelFormat const* const PixelBufferTypes, EPixelFormat const* const DomainBufferTypes)
+FMetalShaderPipeline* FMetalGraphicsPipelineState::GetPipeline(EMetalIndexType IndexType)
 {
 	check(IndexType < EMetalIndexType_Num);
 
-	EMetalBufferType Vertex = VertexShader && (VertexShader->BufferTypeHash && VertexShader->BufferTypeHash == VertexBufferHash) ? EMetalBufferType_Static : EMetalBufferType_Dynamic;
-	EMetalBufferType Fragment = PixelShader && (PixelShader->BufferTypeHash && PixelShader->BufferTypeHash == PixelBufferHash) ? EMetalBufferType_Static : EMetalBufferType_Dynamic;
-	EMetalBufferType Compute = DomainShader && (DomainShader->BufferTypeHash && DomainShader->BufferTypeHash == DomainBufferHash) ? EMetalBufferType_Static : EMetalBufferType_Dynamic;
-
-    FMetalShaderPipeline* Pipe = (GMaxRHIFeatureLevel == ERHIFeatureLevel::SM5) ? PipelineStates[IndexType][Vertex][Fragment][Compute] : nullptr;
-	if ((GMaxRHIFeatureLevel == ERHIFeatureLevel::SM5) && !Pipe)
+	if(!PipelineStates[IndexType])
 	{
-		Pipe = PipelineStates[IndexType][Vertex][Fragment][Compute] = [GetMTLRenderPipeline(true, this, Initializer, IndexType, VertexBufferTypes, PixelBufferTypes, DomainBufferTypes) retain];
+		PipelineStates[IndexType] = [GetMTLRenderPipeline(true, this, Initializer, IndexType) retain];
 	}
-    if (!Pipe)
-    {
-    	if(!PipelineStates[IndexType][0][0][0])
-    	{
-    		PipelineStates[IndexType][0][0][0] = [GetMTLRenderPipeline(true, this, Initializer, IndexType, nullptr, nullptr, nullptr) retain];
-    	}
-        Pipe = PipelineStates[IndexType][0][0][0];
-    }
+	FMetalShaderPipeline* Pipe = PipelineStates[IndexType];
+
 	check(Pipe);
     return Pipe;
 }
diff --git a/Engine/Source/Runtime/Apple/MetalRHI/Private/MetalShaders.cpp b/Engine/Source/Runtime/Apple/MetalRHI/Private/MetalShaders.cpp
--- a/Engine/Source/Runtime/Apple/MetalRHI/Private/MetalShaders.cpp
+++ b/Engine/Source/Runtime/Apple/MetalRHI/Private/MetalShaders.cpp
@@ -35,29 +35,26 @@ struct FMetalCompiledShaderKey
 	FMetalCompiledShaderKey(
 		uint32 InCodeSize,
 		uint32 InCodeCRC,
-		uint32 InConstants,
-		uint32 InBufferTypes
+		uint32 InConstants
 		)
 		: CodeSize(InCodeSize)
 		, CodeCRC(InCodeCRC)
 		, Constants(InConstants)
-		, BufferTypes(InBufferTypes)
 	{}
 
 	friend bool operator ==(const FMetalCompiledShaderKey& A, const FMetalCompiledShaderKey& B)
 	{
-		return A.CodeSize == B.CodeSize && A.CodeCRC == B.CodeCRC && A.Constants == B.Constants && A.BufferTypes == B.BufferTypes;
+		return A.CodeSize == B.CodeSize && A.CodeCRC == B.CodeCRC && A.Constants == B.Constants;
 	}
 
 	friend uint32 GetTypeHash(const FMetalCompiledShaderKey &Key)
 	{
-		return HashCombine(HashCombine(HashCombine(GetTypeHash(Key.CodeSize), GetTypeHash(Key.CodeCRC)), GetTypeHash(Key.Constants)), GetTypeHash(Key.BufferTypes));
+		return HashCombine(HashCombine(GetTypeHash(Key.CodeSize), GetTypeHash(Key.CodeCRC)), GetTypeHash(Key.Constants));
 	}
 
 	uint32 CodeSize;
 	uint32 CodeCRC;
 	uint32 Constants;
-	uint32 BufferTypes;
 };
 
 struct FMetalCompiledShaderCache
@@ -273,10 +270,6 @@ void TMetalBaseShader<BaseResourceType, ShaderType>::Init(const TArray<uint8>& I
 	bHasFunctionConstants = (Header.bTessFunctionConstants || Header.bDeviceFunctionConstants || Header.Bindings.TypedBuffers);
 
 	ConstantValueHash = 0;
-	for (uint8 Constant : Header.Bindings.TypedBufferFormats)
-	{
-		BufferTypeHash ^= Constant;
-	}
 	
 	Library = InLibrary;
 	
@@ -286,12 +279,12 @@ void TMetalBaseShader<BaseResourceType, ShaderType>::Init(const TArray<uint8>& I
 	{
 		// Find the existing compiled shader in the cache.
 		uint32 FunctionConstantHash = i ^ ConstantValueHash;
-		FMetalCompiledShaderKey Key(Header.SourceLen, Header.SourceCRC, FunctionConstantHash, 0);
+		FMetalCompiledShaderKey Key(Header.SourceLen, Header.SourceCRC, FunctionConstantHash);
 		
-		Function[i][EMetalBufferType_Dynamic] = GetMetalCompiledShaderCache().FindRef(Key);
-		if (!Library && Function[i][EMetalBufferType_Dynamic])
+		Function[i] = GetMetalCompiledShaderCache().FindRef(Key);
+		if (!Library && Function[i])
 		{
-			Library = GetMetalCompiledShaderCache().FindLibrary(Function[i][EMetalBufferType_Dynamic]);
+			Library = GetMetalCompiledShaderCache().FindLibrary(Function[i]);
 		}
 		else
 		{
@@ -462,7 +455,7 @@ void TMetalBaseShader<BaseResourceType, ShaderType>::Init(const TArray<uint8>& I
 		
         for (uint32 i = 0; i < Count; i++)
         {
-            GetCompiledFunction((EMetalIndexType)i, nullptr, 0, true);
+            GetCompiledFunction((EMetalIndexType)i, true);
         }
 	}
 	UniformBuffersCopyInfo = Header.UniformBuffersCopyInfo;
@@ -477,38 +470,16 @@ TMetalBaseShader<BaseResourceType, ShaderType>::~TMetalBaseShader()
 }
 
 template<typename BaseResourceType, int32 ShaderType>
-uint32 TMetalBaseShader<BaseResourceType, ShaderType>::GetBufferBindingHash(EPixelFormat const* const BufferTypes) const
-{
-	check(BufferTypes);
-	uint32 VHash = 0;
-	uint BoundBuffers = Bindings.TypedBuffers;
-    while(BoundBuffers)
-    {
-        uint32 Index = __builtin_ctz(BoundBuffers);
-        BoundBuffers &= ~(1 << Index);
-        
-        if (Index < ML_MaxBuffers)
-        {
-        	VHash ^= GMetalBufferFormats[BufferTypes[Index]].DataFormat;
-        }
-    }
-    
-    return VHash;
-}
-
-template<typename BaseResourceType, int32 ShaderType>
-mtlpp::Function TMetalBaseShader<BaseResourceType, ShaderType>::GetCompiledFunction(EMetalIndexType IndexType, EPixelFormat const* const BufferTypes, uint32 InBufferTypeHash, bool const bAsync)
+mtlpp::Function TMetalBaseShader<BaseResourceType, ShaderType>::GetCompiledFunction(EMetalIndexType IndexType, bool const bAsync)
 {
-	EMetalBufferType BT = (InBufferTypeHash == BufferTypeHash && BufferTypeHash) ? EMetalBufferType_Static : EMetalBufferType_Dynamic;
-	
-    mtlpp::Function Func = Function[IndexType][BT];
+	mtlpp::Function Func = Function[IndexType];
     
 	if (!Func)
 	{
 		// Find the existing compiled shader in the cache.
 		uint32 FunctionConstantHash = IndexType ^ ConstantValueHash;
-		FMetalCompiledShaderKey Key(SourceLen, SourceCRC, FunctionConstantHash, (InBufferTypeHash == BufferTypeHash) ? BufferTypeHash : 0);
-		Func = Function[IndexType][BT] = GetMetalCompiledShaderCache().FindRef(Key);
+		FMetalCompiledShaderKey Key(SourceLen, SourceCRC, FunctionConstantHash);
+		Func = Function[IndexType] = GetMetalCompiledShaderCache().FindRef(Key);
 		
 		if (!Func)
 		{
@@ -520,29 +491,6 @@ mtlpp::Function TMetalBaseShader<BaseResourceType, ShaderType>::GetCompiledFunct
             {
                 ConstantValues = mtlpp::FunctionConstantValues();
 				
-				uint32 BoundBuffers = Bindings.TypedBuffers;
-				while(BoundBuffers)
-				{
-					uint32 Index = __builtin_ctz(BoundBuffers);
-					BoundBuffers &= ~(1 << Index);
-					
-					if (Index < ML_MaxBuffers)
-					{
-						if(BT == EMetalBufferType_Static)
-						{
-							// It all matches what is in the shader, so bind the MAX value to force it to elide all the switch-case stuff and just load as directly as possible
-							uint32 v = EMetalBufferFormat::Max;
-							ConstantValues.SetConstantValues(&v, mtlpp::DataType::UInt, ns::Range(Index, 1));
-						}
-						else
-						{
-							// It doesn't match and we don't know what it will be, so load the type dynamically from the buffer meta-table - adds a lot of instructions and hurts performance
-							uint32 v = EMetalBufferFormat::Unknown;
-							ConstantValues.SetConstantValues(&v, mtlpp::DataType::UInt, ns::Range(Index, 1));
-						}
-					}
-				}
-				
                 if(bTessFunctionConstants)
                 {
                     // Index 32 is the tessellation index-buffer presence constant
@@ -560,21 +508,21 @@ mtlpp::Function TMetalBaseShader<BaseResourceType, ShaderType>::GetCompiledFunct
 				METAL_GPUPROFILE(FScopedMetalCPUStats CPUStat(FString::Printf(TEXT("NewFunction: %s"), *FString(Name))));
                 if (!bHasFunctionConstants)
                 {
-                    Function[IndexType][BT] = Library.NewFunction(Name);
+                    Function[IndexType] = Library.NewFunction(Name);
                 }
                 else
                 {
 					ns::AutoReleasedError AError;
-					Function[IndexType][BT] = Library.NewFunction(Name, ConstantValues, &AError);
+					Function[IndexType] = Library.NewFunction(Name, ConstantValues, &AError);
 					ns::Error Error = AError;
-					UE_CLOG(Function[IndexType][BT] == nil, LogMetal, Error, TEXT("Failed to create function: %s"), *FString(Error.GetPtr().description));
-                    UE_CLOG(Function[IndexType][BT] == nil, LogMetal, Fatal, TEXT("*********** Error\n%s"), *FString(GetSourceCode()));
+					UE_CLOG(Function[IndexType] == nil, LogMetal, Error, TEXT("Failed to create function: %s"), *FString(Error.GetPtr().description));
+                    UE_CLOG(Function[IndexType] == nil, LogMetal, Fatal, TEXT("*********** Error\n%s"), *FString(GetSourceCode()));
                 }
                 
-                check(Function[IndexType][BT]);
-                GetMetalCompiledShaderCache().Add(Key, Library, Function[IndexType][BT]);
+                check(Function[IndexType]);
+                GetMetalCompiledShaderCache().Add(Key, Library, Function[IndexType]);
                 
-                Func = Function[IndexType][BT];
+                Func = Function[IndexType];
             }
             else
             {
@@ -613,7 +561,7 @@ FMetalComputeShader::FMetalComputeShader(const TArray<uint8>& InCode, mtlpp::Lib
 , NumThreadsY(0)
 , NumThreadsZ(0)
 {
-    Pipeline[0] = Pipeline[1] = nil;
+    Pipeline = nil;
 	FMetalCodeHeader Header;
 	Init(InCode, Header, InLibrary);
 	
@@ -624,29 +572,15 @@ FMetalComputeShader::FMetalComputeShader(const TArray<uint8>& InCode, mtlpp::Lib
 
 FMetalComputeShader::~FMetalComputeShader()
 {
-    for (uint32 i = 0; i < EMetalBufferType_Num; i++)
-    {
-        [Pipeline[i] release];
-        Pipeline[i] = nil;
-    }
-}
-	
-uint32 FMetalComputeShader::GetBindingHash(EPixelFormat const* const BufferTypes) const
-{
-	if (BufferTypes)
-	{
-		return GetBufferBindingHash(BufferTypes);
-	}
-	return 0;
+	[Pipeline release];
+	Pipeline = nil;
 }
 
-FMetalShaderPipeline* FMetalComputeShader::GetPipeline(EPixelFormat const* const BufferTypes, uint32 InBufferTypeHash)
+FMetalShaderPipeline* FMetalComputeShader::GetPipeline()
 {
-	EMetalBufferType BT = (InBufferTypeHash == BufferTypeHash && BufferTypeHash) ? EMetalBufferType_Static : EMetalBufferType_Dynamic;
-	
-	if (!Pipeline[BT])
+	if (!Pipeline)
 	{
-		mtlpp::Function Func = GetCompiledFunction(EMetalIndexType_None, BufferTypes, (InBufferTypeHash == BufferTypeHash) ? BufferTypeHash : 0);
+		mtlpp::Function Func = GetCompiledFunction(EMetalIndexType_None);
 		check(Func);
         
 		ns::Error Error;
@@ -710,21 +644,21 @@ FMetalShaderPipeline* FMetalComputeShader::GetPipeline(EPixelFormat const* const
             UE_LOG(LogRHI, Fatal, TEXT("Failed to create compute kernel: %s"), *FString([Error description]));
         }
         
-        Pipeline[BT] = [FMetalShaderPipeline new];
-        Pipeline[BT]->ComputePipelineState = Kernel;
+        Pipeline = [FMetalShaderPipeline new];
+        Pipeline->ComputePipelineState = Kernel;
 #if METAL_DEBUG_OPTIONS
-        Pipeline[BT]->ComputePipelineReflection = Reflection;
-        Pipeline[BT]->ComputeSource = GetSourceCode();
+        Pipeline->ComputePipelineReflection = Reflection;
+        Pipeline->ComputeSource = GetSourceCode();
 		if (Reflection)
 		{
-			Pipeline[BT]->ComputeDesc = Descriptor;
+			Pipeline->ComputeDesc = Descriptor;
 		}
 #endif
-        METAL_DEBUG_OPTION(FMemory::Memzero(Pipeline[BT]->ResourceMask, sizeof(Pipeline[BT]->ResourceMask)));
+        METAL_DEBUG_OPTION(FMemory::Memzero(Pipeline->ResourceMask, sizeof(Pipeline->ResourceMask)));
 	}
-	check(Pipeline[BT]);
+	check(Pipeline);
 
-	return Pipeline[BT];
+	return Pipeline;
 }
 
 FMetalVertexShader::FMetalVertexShader(const TArray<uint8>& InCode)
@@ -764,19 +698,10 @@ FMetalVertexShader::FMetalVertexShader(const TArray<uint8>& InCode, mtlpp::Libra
 	TessellationMaxTessFactor = Header.TessellationMaxTessFactor;
 	TessellationPatchesPerThreadGroup = Header.TessellationPatchesPerThreadGroup;
 }
-	
-uint32 FMetalVertexShader::GetBindingHash(EPixelFormat const* const BufferTypes) const
-{
-	if (BufferTypes)
-	{
-		return GetBufferBindingHash(BufferTypes);
-	}
-	return 0;
-}
 
-mtlpp::Function FMetalVertexShader::GetFunction(EMetalIndexType IndexType, EPixelFormat const* const BufferTypes, uint32 BufferTypeHash)
+mtlpp::Function FMetalVertexShader::GetFunction(EMetalIndexType IndexType)
 {
-	return GetCompiledFunction(IndexType, BufferTypes, BufferTypeHash);
+	return GetCompiledFunction(IndexType);
 }
 
 FMetalPixelShader::FMetalPixelShader(const TArray<uint8>& InCode)
@@ -790,19 +715,10 @@ FMetalPixelShader::FMetalPixelShader(const TArray<uint8>& InCode, mtlpp::Library
 	FMetalCodeHeader Header;
 	Init(InCode, Header, InLibrary);
 }
-	
-uint32 FMetalPixelShader::GetBindingHash(EPixelFormat const* const BufferTypes) const
-{
-	if (BufferTypes)
-	{
-		return GetBufferBindingHash(BufferTypes);
-	}
-	return 0;
-}
 
-mtlpp::Function FMetalPixelShader::GetFunction(EMetalIndexType IndexType, EPixelFormat const* const BufferTypes, uint32 BufferTypeHash)
+mtlpp::Function FMetalPixelShader::GetFunction(EMetalIndexType IndexType)
 {
-	return GetCompiledFunction(IndexType, BufferTypes, BufferTypeHash);
+	return GetCompiledFunction(IndexType);
 }
 
 FMetalHullShader::FMetalHullShader(const TArray<uint8>& InCode)
@@ -816,19 +732,10 @@ FMetalHullShader::FMetalHullShader(const TArray<uint8>& InCode, mtlpp::Library I
 	FMetalCodeHeader Header;
 	Init(InCode, Header, InLibrary);
 }
-	
-uint32 FMetalHullShader::GetBindingHash(EPixelFormat const* const BufferTypes) const
-{
-	if (BufferTypes)
-	{
-		return GetBufferBindingHash(BufferTypes);
-	}
-	return 0;
-}
 
-mtlpp::Function FMetalHullShader::GetFunction(EMetalIndexType IndexType, EPixelFormat const* const BufferTypes, uint32 BufferTypeHash)
+mtlpp::Function FMetalHullShader::GetFunction(EMetalIndexType IndexType)
 {
-	return GetCompiledFunction(IndexType, BufferTypes, BufferTypeHash);
+	return GetCompiledFunction(IndexType);
 }
 
 FMetalDomainShader::FMetalDomainShader(const TArray<uint8>& InCode)
@@ -884,19 +791,10 @@ FMetalDomainShader::FMetalDomainShader(const TArray<uint8>& InCode, mtlpp::Libra
 		default: check(0);
 	}
 }
-	
-uint32 FMetalDomainShader::GetBindingHash(EPixelFormat const* const BufferTypes) const
-{
-	if (BufferTypes)
-	{
-		return GetBufferBindingHash(BufferTypes);
-	}
-	return 0;
-}
 
-mtlpp::Function FMetalDomainShader::GetFunction(EMetalIndexType IndexType, EPixelFormat const* const BufferTypes, uint32 BufferTypeHash)
+mtlpp::Function FMetalDomainShader::GetFunction(EMetalIndexType IndexType)
 {
-	return GetCompiledFunction(IndexType, BufferTypes, BufferTypeHash);
+	return GetCompiledFunction(IndexType);
 }
 
 FVertexShaderRHIRef FMetalDynamicRHI::RHICreateVertexShader(const TArray<uint8>& Code)
@@ -1128,7 +1026,7 @@ FPixelShaderRHIRef FMetalShaderLibrary::CreatePixelShader(const FSHAHash& Hash)
 	if (Code)
 	{
 		FMetalPixelShader* Shader = new FMetalPixelShader(Code->Data, Library[Code->Index]);
-		if (Shader->GetFunction(EMetalIndexType_None, nullptr, 0))
+		if (Shader->GetFunction(EMetalIndexType_None))
 		{
 			return Shader;
 		}
@@ -1148,7 +1046,7 @@ FVertexShaderRHIRef FMetalShaderLibrary::CreateVertexShader(const FSHAHash& Hash
 	if (Code)
 	{
 		FMetalVertexShader* Shader = new FMetalVertexShader(Code->Data, Library[Code->Index]);
-		if (Shader->GetFunction(EMetalIndexType_None, nullptr, 0))
+		if (Shader->GetFunction(EMetalIndexType_None))
 		{
 			return Shader;
 		}
@@ -1167,7 +1065,7 @@ FHullShaderRHIRef FMetalShaderLibrary::CreateHullShader(const FSHAHash& Hash)
 	if (Code)
 	{
 		FMetalHullShader* Shader = new FMetalHullShader(Code->Data, Library[Code->Index]);
-		if(Shader->GetFunction(EMetalIndexType_None, nullptr, 0))
+		if(Shader->GetFunction(EMetalIndexType_None))
 		{
 			return Shader;
 		}
@@ -1186,7 +1084,7 @@ FDomainShaderRHIRef FMetalShaderLibrary::CreateDomainShader(const FSHAHash& Hash
 	if (Code)
 	{
 		FMetalDomainShader* Shader = new FMetalDomainShader(Code->Data, Library[Code->Index]);
-		if (Shader->GetFunction(EMetalIndexType_None, nullptr, 0))
+		if (Shader->GetFunction(EMetalIndexType_None))
 		{
 			return Shader;
 		}
@@ -1217,7 +1115,7 @@ FComputeShaderRHIRef FMetalShaderLibrary::CreateComputeShader(const FSHAHash& Ha
 	if (Code)
 	{
 		FMetalComputeShader* Shader = new FMetalComputeShader(Code->Data, Library[Code->Index]);
-		if (Shader->GetPipeline(nullptr, 0))
+		if (Shader->GetPipeline())
 		{
 			return Shader;
 		}
diff --git a/Engine/Source/Runtime/Apple/MetalRHI/Private/MetalStateCache.cpp b/Engine/Source/Runtime/Apple/MetalRHI/Private/MetalStateCache.cpp
--- a/Engine/Source/Runtime/Apple/MetalRHI/Private/MetalStateCache.cpp
+++ b/Engine/Source/Runtime/Apple/MetalRHI/Private/MetalStateCache.cpp
@@ -208,7 +208,6 @@ FMetalStateCache::~FMetalStateCache()
 			ShaderBuffers[Frequency].Formats[i] = PF_Unknown;
 		}
 		ShaderBuffers[Frequency].Bound = 0;
-		ShaderBuffers[Frequency].FormatHash = 0;
 		for (uint32 i = 0; i < ML_MaxTextures; i++)
 		{
 			ShaderTextures[Frequency].Textures[i] = nil;
@@ -264,7 +263,6 @@ void FMetalStateCache::Reset(void)
 			ShaderBuffers[Frequency].Formats[i] = PF_Unknown;
 		}
 		ShaderBuffers[Frequency].Bound = 0;
-		ShaderBuffers[Frequency].FormatHash = 0;
 		for (uint32 i = 0; i < ML_MaxTextures; i++)
 		{
 			ShaderTextures[Frequency].Textures[i] = nil;
@@ -1066,7 +1064,7 @@ void FMetalStateCache::SetGraphicsPipelineState(FMetalGraphicsPipelineState* Sta
 	{
 		GraphicsPSO = State;
 		
-		bool bNewUsingTessellation = (State && State->GetPipeline(IndexType, EMetalBufferType_Dynamic, EMetalBufferType_Dynamic, EMetalBufferType_Dynamic)->TessellationPipelineDesc.DomainVertexDescriptor);
+		bool bNewUsingTessellation = (State && State->GetPipeline(IndexType)->TessellationPipelineDesc.DomainVertexDescriptor);
 		if (bNewUsingTessellation != bUsingTessellation)
 		{
 			for (uint32 i = 0; i < SF_NumFrequencies; i++)
@@ -1094,7 +1092,6 @@ void FMetalStateCache::SetGraphicsPipelineState(FMetalGraphicsPipelineState* Sta
 			ShaderTextures[SF_Hull].Bound = UINT32_MAX;
 #endif
 			ShaderSamplers[SF_Hull].Bound = UINT16_MAX;
-			ShaderBuffers[SF_Hull].FormatHash = 0;
 			
 			for (uint32 i = 0; i < ML_MaxBuffers; i++)
 			{
@@ -1351,16 +1348,6 @@ bool FMetalStateCache::NeedsToSetRenderTarget(const FRHISetRenderTargetsInfo& In
 	return bAllChecksPassed == false;
 }
 
-static uint8 GMetalShaderFreqFormat[SF_NumFrequencies] =
-{
-	EMetalPipelineFlagVertexBuffers,
-	0,
-	EMetalPipelineFlagDomainBuffers,
-	EMetalPipelineFlagPixelBuffers,
-	0,
-	EMetalPipelineFlagComputeBuffers
-};
-
 void FMetalStateCache::SetShaderBuffer(EShaderFrequency const Frequency, FMetalBuffer const& Buffer, FMetalBufferData* const Bytes, NSUInteger const Offset, NSUInteger const Length, NSUInteger const Index, mtlpp::ResourceUsage const Usage, EPixelFormat const Format)
 {
 	check(Frequency < SF_NumFrequencies);
@@ -1379,7 +1366,6 @@ void FMetalStateCache::SetShaderBuffer(EShaderFrequency const Frequency, FMetalB
 		ShaderBuffers[Frequency].Buffers[Index].Length = Length;
 		ShaderBuffers[Frequency].Buffers[Index].Usage = Usage;
 		
-		PipelineBits |= (ShaderBuffers[Frequency].Formats[Index] != Format) ? GMetalShaderFreqFormat[Frequency] : 0;
 		ShaderBuffers[Frequency].Formats[Index] = Format;
 		
 		if (Buffer || Bytes)
@@ -1522,12 +1508,10 @@ void FMetalStateCache::SetShaderResourceView(FMetalContext* Context, EShaderFreq
 		}
 		else if (VB)
 		{
-			checkf(ValidateBufferFormat(ShaderStage, BindIndex, (EPixelFormat)SRV->Format), TEXT("Invalid buffer format %d for index %d, shader %d"), SRV->Format, BindIndex, ShaderStage);
 			SetShaderBuffer(ShaderStage, VB->Buffer, VB->Data, 0, VB->GetSize(), BindIndex, mtlpp::ResourceUsage::Read, (EPixelFormat)SRV->Format);
 		}
 		else if (IB)
 		{
-			checkf(ValidateBufferFormat(ShaderStage, BindIndex, (EPixelFormat)SRV->Format), TEXT("Invalid buffer format %d for index %d, shader %d"), SRV->Format, BindIndex, ShaderStage);
 			SetShaderBuffer(ShaderStage, IB->Buffer, nil, 0, IB->GetSize(), BindIndex, mtlpp::ResourceUsage::Read, (EPixelFormat)SRV->Format);
 		}
 		else if (SB)
@@ -1573,42 +1557,6 @@ bool FMetalStateCache::IsLinearBuffer(EShaderFrequency ShaderStage, uint32 BindI
     }
 }
 
-bool FMetalStateCache::ValidateBufferFormat(EShaderFrequency ShaderStage, uint32 BindIndex, EPixelFormat Format)
-{
-	switch (ShaderStage)
-	{
-		case SF_Vertex:
-		{
-			return (GraphicsPSO->VertexShader->Bindings.InvariantBuffers & (1 << BindIndex)) == 0 || (GMetalBufferFormats[Format].DataFormat == GraphicsPSO->VertexShader->Bindings.TypedBufferFormats[BindIndex]);
-			break;
-		}
-		case SF_Pixel:
-		{
-			return (GraphicsPSO->PixelShader->Bindings.InvariantBuffers & (1 << BindIndex)) == 0 || (GMetalBufferFormats[Format].DataFormat == GraphicsPSO->PixelShader->Bindings.TypedBufferFormats[BindIndex]);
-			break;
-		}
-		case SF_Hull:
-		{
-			return (GraphicsPSO->HullShader->Bindings.InvariantBuffers & (1 << BindIndex)) == 0 || (GMetalBufferFormats[Format].DataFormat == GraphicsPSO->HullShader->Bindings.TypedBufferFormats[BindIndex]);
-			break;
-		}
-		case SF_Domain:
-		{
-			return (GraphicsPSO->DomainShader->Bindings.InvariantBuffers & (1 << BindIndex)) == 0 || (GMetalBufferFormats[Format].DataFormat == GraphicsPSO->DomainShader->Bindings.TypedBufferFormats[BindIndex]);
-			break;
-		}
-		case SF_Compute:
-		{
-			return (ComputeShader->Bindings.InvariantBuffers & (1 << BindIndex)) == 0 || (GMetalBufferFormats[Format].DataFormat == ComputeShader->Bindings.TypedBufferFormats[BindIndex]);
-		}
-		default:
-		{
-			check(false);
-			return false;
-		}
-	}
-}
-
 void FMetalStateCache::SetShaderUnorderedAccessView(EShaderFrequency ShaderStage, uint32 BindIndex, FMetalUnorderedAccessView* RESTRICT UAV)
 {
 	if (UAV)
@@ -1637,7 +1585,6 @@ void FMetalStateCache::SetShaderUnorderedAccessView(EShaderFrequency ShaderStage
 			}
 			else
 			{
-				checkf(ValidateBufferFormat(ShaderStage, BindIndex, (EPixelFormat)UAV->SourceView->Format), TEXT("Invalid buffer format %d for index %d, shader %d"), UAV->SourceView->Format, BindIndex, ShaderStage);
 				SetShaderBuffer(ShaderStage, VertexBuffer->Buffer, VertexBuffer->Data, 0, VertexBuffer->GetSize(), BindIndex, mtlpp::ResourceUsage(mtlpp::ResourceUsage::Read | mtlpp::ResourceUsage::Write), (EPixelFormat)UAV->SourceView->Format);
 			}
 		}
@@ -1655,7 +1602,6 @@ void FMetalStateCache::SetShaderUnorderedAccessView(EShaderFrequency ShaderStage
 			}
 			else
 			{
-				checkf(ValidateBufferFormat(ShaderStage, BindIndex, (EPixelFormat)UAV->SourceView->Format), TEXT("Invalid buffer format %d for index %d, shader %d"), UAV->SourceView->Format, BindIndex, ShaderStage);
 				SetShaderBuffer(ShaderStage, IndexBuffer->Buffer, nullptr, 0, IndexBuffer->GetSize(), BindIndex, mtlpp::ResourceUsage(mtlpp::ResourceUsage::Read | mtlpp::ResourceUsage::Write), (EPixelFormat)UAV->SourceView->Format);
 			}
 		}
@@ -2113,36 +2059,10 @@ void FMetalStateCache::SetRenderPipelineState(FMetalCommandEncoder& CommandEncod
 	
     if ((PipelineBits & EMetalPipelineFlagRasterMask) != 0)
     {
-    	// @todo Could optimise it so that we only re-evaluate the buffer hashes if the shader buffer binding mask changes when changing the PSO
-        if (PipelineBits & EMetalPipelineFlagPipelineState)
-        {
-        	PipelineBits |= (EMetalPipelineFlagVertexBuffers|EMetalPipelineFlagPixelBuffers|EMetalPipelineFlagDomainBuffers);
-        }
-    	
-    	if (PipelineBits & EMetalPipelineFlagVertexBuffers)
-    	{
-    		ShaderBuffers[SF_Vertex].FormatHash = GraphicsPSO->VertexShader->GetBindingHash(ShaderBuffers[SF_Vertex].Formats);
-    	}
-    	
-    	if (PipelineBits & EMetalPipelineFlagPixelBuffers)
-    	{
-   			ShaderBuffers[SF_Pixel].FormatHash = (IsValidRef(GraphicsPSO->PixelShader)) ? GraphicsPSO->PixelShader->GetBindingHash(ShaderBuffers[SF_Pixel].Formats) : 0;
-    	}
-    	
-    	if (PipelineBits & EMetalPipelineFlagDomainBuffers)
-    	{
-   			ShaderBuffers[SF_Domain].FormatHash = (IsValidRef(GraphicsPSO->DomainShader)) ? GraphicsPSO->DomainShader->GetBindingHash(ShaderBuffers[SF_Domain].Formats) : 0;
-    	}
-    
-    	EPixelFormat const* const VertexFormats = ShaderBuffers[SF_Vertex].Formats;
-    	EPixelFormat const* const PixelFormats = (IsValidRef(GraphicsPSO->PixelShader)) ? ShaderBuffers[SF_Pixel].Formats : nullptr;
-    	EPixelFormat const* const DomainFormats = (IsValidRef(GraphicsPSO->DomainShader)) ? ShaderBuffers[SF_Domain].Formats : nullptr;
-    
-        // Some Intel drivers need RenderPipeline state to be set after DepthStencil state to work properly
-        // As it happens, in order to use function constants to emulate Buffer<T>/RWBuffer<T> implicit typing we'll do that anyway.
-    	FMetalShaderPipeline* Pipeline = GetPipelineState(ShaderBuffers[SF_Vertex].FormatHash, ShaderBuffers[SF_Pixel].FormatHash, ShaderBuffers[SF_Domain].FormatHash, VertexFormats, PixelFormats, DomainFormats);
-        // FMetalShaderPipeline* Pipeline = GetPipelineState(0, 0, 0, nullptr, nullptr, nullptr);
-        check(Pipeline);
+    	// Some Intel drivers need RenderPipeline state to be set after DepthStencil state to work properly
+    	FMetalShaderPipeline* Pipeline = GetPipelineState();
+
+		check(Pipeline);
         CommandEncoder.SetRenderPipelineState(Pipeline);
         if (Pipeline->ComputePipelineState)
         {
@@ -2158,17 +2078,7 @@ void FMetalStateCache::SetComputePipelineState(FMetalCommandEncoder& CommandEnco
 {
 	if ((PipelineBits & EMetalPipelineFlagComputeMask) != 0)
 	{
-		if (PipelineBits & EMetalPipelineFlagComputeShader)
-		{
-			PipelineBits |= EMetalPipelineFlagComputeBuffers;
-		}
-	
-		if (PipelineBits & EMetalPipelineFlagComputeBuffers)
-		{
-			ShaderBuffers[SF_Compute].FormatHash = ComputeShader->GetBindingHash(ShaderBuffers[SF_Compute].Formats);
-		}
-	    
-	    FMetalShaderPipeline* Pipeline = ComputeShader->GetPipeline(ShaderBuffers[SF_Compute].Formats, ShaderBuffers[SF_Compute].FormatHash);
+		FMetalShaderPipeline* Pipeline = ComputeShader->GetPipeline();
 	    check(Pipeline);
 	    CommandEncoder.SetComputePipelineState(Pipeline);
         
diff --git a/Engine/Source/Runtime/Apple/MetalRHI/Private/MetalStateCache.h b/Engine/Source/Runtime/Apple/MetalRHI/Private/MetalStateCache.h
--- a/Engine/Source/Runtime/Apple/MetalRHI/Private/MetalStateCache.h
+++ b/Engine/Source/Runtime/Apple/MetalRHI/Private/MetalStateCache.h
@@ -9,11 +9,7 @@
 enum EMetalPipelineFlags
 {
 	EMetalPipelineFlagPipelineState = 1 << 0,
-    EMetalPipelineFlagVertexBuffers = 1 << 1,
-    EMetalPipelineFlagPixelBuffers = 1 << 2,
-    EMetalPipelineFlagDomainBuffers = 1 << 3,
-    EMetalPipelineFlagComputeBuffers = 1 << 4,
-    EMetalPipelineFlagComputeShader = 1 << 5,
+    EMetalPipelineFlagComputeShader = 1 << 1,
     EMetalPipelineFlagRasterMask = 0xF,
     EMetalPipelineFlagComputeMask = 0x30,
     EMetalPipelineFlagMask = 0x3F
@@ -143,9 +139,7 @@ public:
 	mtlpp::RenderPassDescriptor GetRenderPassDescriptor(void) const { return RenderPassDesc; }
 	uint32 GetSampleCount(void) const { return SampleCount; }
     bool IsLinearBuffer(EShaderFrequency ShaderStage, uint32 BindIndex);
-	bool ValidateBufferFormat(EShaderFrequency ShaderStage, uint32 BindIndex, EPixelFormat Format);
-    FMetalShaderPipeline* GetPipelineState(uint32 V, uint32 F, uint32 C, EPixelFormat const* const VS, EPixelFormat const* const PS, EPixelFormat const* const DS) const { return GraphicsPSO->GetPipeline(GetIndexType(), V, F, C, VS, PS, DS); }
-    FMetalShaderPipeline* GetPipelineState(void) const { return GraphicsPSO->GetPipeline(GetIndexType(), ShaderBuffers[SF_Vertex].FormatHash, ShaderBuffers[SF_Pixel].FormatHash, ShaderBuffers[SF_Domain].FormatHash, nullptr, nullptr, nullptr); }
+    FMetalShaderPipeline* GetPipelineState(void) const { return GraphicsPSO->GetPipeline(GetIndexType()); }
 	EPrimitiveType GetPrimitiveType() { check(IsValidRef(GraphicsPSO)); return GraphicsPSO->GetPrimitiveType(); }
 	
 	FTexture2DRHIRef CreateFallbackDepthStencilSurface(uint32 Width, uint32 Height);
@@ -199,13 +193,11 @@ private:
 	/** A structure of arrays for the current buffer binding settings. */
 	struct FMetalBufferBindings
 	{
-		FMetalBufferBindings() : FormatHash(0), Bound(0) {}
+		FMetalBufferBindings() : Bound(0) {}
 		/** The bound buffers/bytes or nil. */
 		FMetalBufferBinding Buffers[ML_MaxBuffers];
 		/** The pixel formats for buffers bound so that we emulate [RW]Buffer<T> type conversion */
 		EPixelFormat Formats[ML_MaxBuffers];
-		/** The hash of the pixel formats for the formats above */
-		uint32 FormatHash;
 		/** A bitmask for which buffers were bound by the application where a bit value of 1 is bound and 0 is unbound. */
 		uint32 Bound;
 	};
@@ -231,8 +223,6 @@ private:
 		/** A bitmask for which samplers were bound by the application where a bit value of 1 is bound and 0 is unbound. */
 		uint16 Bound;
 	};
-	
-    EMetalBufferType GetShaderBufferBindingType(FMetalBufferBindings& BufferBindings, uint32 BoundBuffers, uint32 ShaderBindingHash);
     
 private:
 	FMetalShaderParameterCache ShaderParameters[CrossCompiler::NUM_SHADER_STAGES];
diff --git a/Engine/Source/Runtime/Apple/MetalRHI/Private/MetalVertexDeclaration.cpp b/Engine/Source/Runtime/Apple/MetalRHI/Private/MetalVertexDeclaration.cpp
--- a/Engine/Source/Runtime/Apple/MetalRHI/Private/MetalVertexDeclaration.cpp
+++ b/Engine/Source/Runtime/Apple/MetalRHI/Private/MetalVertexDeclaration.cpp
@@ -7,6 +7,8 @@
 #include "MetalRHIPrivate.h"
 #include "MetalProfiler.h"
 
+mtlpp::VertexFormat GMetalFColorVertexFormat = mtlpp::VertexFormat::UChar4Normalized;
+
 static mtlpp::VertexFormat TranslateElementTypeToMTLType(EVertexElementType Type)
 {
 	switch (Type)
@@ -18,7 +20,7 @@ static mtlpp::VertexFormat TranslateElementTypeToMTLType(EVertexElementType Type
 		case VET_PackedNormal:	return mtlpp::VertexFormat::Char4Normalized;
 		case VET_UByte4:		return mtlpp::VertexFormat::UChar4;
 		case VET_UByte4N:		return mtlpp::VertexFormat::UChar4Normalized;
-		case VET_Color:			return mtlpp::VertexFormat::UChar4Normalized;
+		case VET_Color:			return GMetalFColorVertexFormat;
 		case VET_Short2:		return mtlpp::VertexFormat::Short2;
 		case VET_Short4:		return mtlpp::VertexFormat::Short4;
 		case VET_Short2N:		return mtlpp::VertexFormat::Short2Normalized;
