Remove a lot of dead-code for Metal tessellation that hasn't been touched in the last few years and that won't be revived at this stage.

diff --git a/Engine/Source/Developer/Apple/MetalShaderFormat/Private/MetalBackend.cpp b/Engine/Source/Developer/Apple/MetalShaderFormat/Private/MetalBackend.cpp
--- a/Engine/Source/Developer/Apple/MetalShaderFormat/Private/MetalBackend.cpp
+++ b/Engine/Source/Developer/Apple/MetalShaderFormat/Private/MetalBackend.cpp
@@ -53,11 +53,6 @@ PRAGMA_ENABLE_SHADOW_VARIABLE_WARNINGS
 #define WAVE_PREFIX_SUM "WavePrefixSum"
 #define WAVE_PREFIX_PRODUCT "WavePrefixProduct"
 
-// NOTE: a lot of the comments refer to running out OUTPUT_CP rate -- not all comments were fixed...
-#define EXEC_AT_INPUT_CP_RATE 1 // exec at input CP rate
-
-#define MULTI_PATCH 1
-
 /**
  * This table must match the ir_expression_operation enum.
  */
@@ -1622,7 +1617,6 @@ protected:
 				{
 					if (Backend.bIsTessellationVSHS)
 					{
-						check(EXEC_AT_INPUT_CP_RATE);
 						ralloc_asprintf_append(buffer, "#define GET_PATCH_COUNT() patchCount[0]\n");
 						ralloc_asprintf_append(buffer, "#define GET_PATCH_ID() (thread_position_in_grid.x / TessellationInputControlPoints)\n");
 						ralloc_asprintf_append(buffer, "#define GET_PATCH_VALID() (GET_PATCH_ID() < GET_PATCH_COUNT())\n");
@@ -5189,7 +5183,6 @@ bool FMetalCodeBackend::GenerateMain(EHlslShaderFrequency Frequency, const char*
 		check(patchesPerThreadgroup != 0);
 		check(patchesPerThreadgroup <= METAL_TESS_MAX_THREADS_PER_THREADGROUP);
 
-#if EXEC_AT_INPUT_CP_RATE
 		// create and call GET_INPUT_CP_ID
 		ir_variable* SV_InputControlPointIDVar = NULL; // @todo it would be better to do this under GenerateInputFromSemantic (also this is ... should never be used by anything in the USF... only internal)
 		{
@@ -5235,35 +5228,6 @@ bool FMetalCodeBackend::GenerateMain(EHlslShaderFrequency Frequency, const char*
 				);
 			PrePreCallInstructions.push_tail(Assign);
 		}
-#else // EXEC_AT_INPUT_CP_RATE
-		// create and call GET_OUTPUT_CP_ID
-		ir_variable* SV_OutputControlPointIDVar = NULL; // @todo it would be better to do this under GenerateInputFromSemantic
-		{
-			ir_function *Function = NULL;
-			// create GET_OUTPUT_CP_ID
-			{
-				const glsl_type* retType = glsl_type::get_instance(GLSL_TYPE_UINT, 1, 1);
-				ir_function_signature* sig = new(ParseState)ir_function_signature(retType);
-				sig->is_builtin = true;
-				Function = new(ParseState)ir_function("GET_OUTPUT_CP_ID");
-				Function->add_signature(sig);
-			}
-			check(Function);
-
-			exec_list VoidParameter;
-			ir_function_signature * FunctionSig = Function->matching_signature(&VoidParameter);
-
-			ir_variable* TempVariable = new(ParseState) ir_variable(glsl_type::get_instance(GLSL_TYPE_UINT, 1, 1), "SV_OutputControlPointID", ir_var_temporary);
-			ir_dereference_variable* TempVariableDeref = new(ParseState) ir_dereference_variable(TempVariable);
-			PrePreCallInstructions.push_tail(TempVariable);
-
-			auto* Call = new(ParseState)ir_call(FunctionSig, TempVariableDeref, &VoidParameter);
-			PrePreCallInstructions.push_tail(Call);
-
-			SV_OutputControlPointIDVar = TempVariable;
-			ParseState->symbols->add_variable(SV_OutputControlPointIDVar);
-		}
-#endif // EXEC_AT_INPUT_CP_RATE
 		// create and call GET_PATCH_VALID
 		{
 			ir_function *Function = NULL;
@@ -5787,7 +5751,6 @@ bool FMetalCodeBackend::GenerateMain(EHlslShaderFrequency Frequency, const char*
 
 
 			// optimization if inputcontrolpoints == outputcontrolpoints -- no need for a loop
-			if(EXEC_AT_INPUT_CP_RATE || inputcontrolpoints == ParseState->tessellation.outputcontrolpoints)
 			{
 				// add ... if(isPatchValid)
 				ir_if* pv_if = new(ParseState)ir_if(new(ParseState)ir_dereference_variable(ParseState->symbols->get_variable("isPatchValid")));
@@ -5805,11 +5768,7 @@ bool FMetalCodeBackend::GenerateMain(EHlslShaderFrequency Frequency, const char*
 						InputPatchVar,
 						new (ParseState)ir_dereference_variable(ParseState->symbols->get_variable("patchIDInThreadgroup"))
 					),
-#if EXEC_AT_INPUT_CP_RATE
 					new (ParseState)ir_dereference_variable(ParseState->symbols->get_variable("SV_InputControlPointID"))
-#else // EXEC_AT_INPUT_CP_RATE
-					new (ParseState)ir_dereference_variable(ParseState->symbols->get_variable("SV_OutputControlPointID"))
-#endif // EXEC_AT_INPUT_CP_RATE
 					);
 				pv_if->then_instructions.push_tail(
 					new (ParseState)ir_assignment(
@@ -5818,106 +5777,6 @@ bool FMetalCodeBackend::GenerateMain(EHlslShaderFrequency Frequency, const char*
 					)
 					);
 			}
-			else
-			{
-				check(0); // not currently a supported combination with compute stageIn attributes
-				check(!EXEC_AT_INPUT_CP_RATE); // this will never happen if EXEC_AT_INPUT_CP_RATE
-				// add ...	for(uint baseCPID = 0; baseCPID < TessellationInputControlPoints; baseCPID += TessellationOutputControlPoints)
-				ir_variable* baseCPIDVar = new(ParseState)ir_variable(glsl_type::get_instance(GLSL_TYPE_UINT, 1, 1), "baseCPIDVar", ir_var_temporary);
-				PrePreCallInstructions.push_tail(baseCPIDVar);
-				// add ... uint baseCPID = 0
-				PrePreCallInstructions.push_tail(
-					new(ParseState)ir_assignment(
-						new(ParseState)ir_dereference_variable(baseCPIDVar),
-						new(ParseState)ir_constant((unsigned)0)
-					)
-					);
-				ir_loop* vf_loop = new(ParseState)ir_loop();
-				PrePreCallInstructions.push_tail(vf_loop);
-
-				// NOTE: cannot use from/to/increment/counter/cmp because that is used during optimizations
-				// add ... baseCPID < TessellationInputControlPoints (to break from the for loop)
-				ir_if* vf_loop_break = new(ParseState)ir_if(
-					new(ParseState)ir_expression(ir_binop_gequal,
-						new(ParseState)ir_dereference_variable(baseCPIDVar),
-						new(ParseState)ir_constant((unsigned)inputcontrolpoints)
-					)
-					);
-				vf_loop->body_instructions.push_tail(vf_loop_break);
-				vf_loop_break->then_instructions.push_tail(
-						new(ParseState)ir_loop_jump(ir_loop_jump::jump_break)
-					);
-				vf_loop->mode = ir_loop::loop_dont_care;
-
-				// add ... const uint inputCPID = baseCPID + SV_OutputControlPointID; // baseCPID + GET_OUTPUT_CP_ID()
-				ir_variable* inputCPIDVar = new(ParseState)ir_variable(glsl_type::get_instance(GLSL_TYPE_UINT, 1, 1), "inputCPIDVar", ir_var_temporary);
-				vf_loop->body_instructions.push_tail(inputCPIDVar);
-				vf_loop->body_instructions.push_tail(
-					new(ParseState)ir_assignment(
-						new(ParseState)ir_dereference_variable(inputCPIDVar),
-						new(ParseState)ir_expression(ir_binop_add,
-							new(ParseState)ir_dereference_variable(baseCPIDVar),
-							new(ParseState)ir_dereference_variable(ParseState->symbols->get_variable("SV_OutputControlPointID"))
-						)
-					)
-					);
-
-				// add ... if(inputCPID < TessellationInputControlPoints)
-				ir_if* vf_if = new(ParseState)ir_if(
-					new(ParseState)ir_expression(ir_binop_less,
-						new(ParseState)ir_dereference_variable(inputCPIDVar),
-						new(ParseState)ir_constant((unsigned)inputcontrolpoints)
-					)
-					);
-				vf_loop->body_instructions.push_tail(vf_if);
-				// add ... baseCPID += TessellationOutputControlPoints
-				vf_loop->body_instructions.push_tail(
-					new(ParseState)ir_assignment(
-						new(ParseState)ir_dereference_variable(baseCPIDVar),
-						new(ParseState)ir_expression(ir_binop_add,
-							new(ParseState)ir_dereference_variable(baseCPIDVar),
-							new(ParseState)ir_constant((unsigned)ParseState->tessellation.outputcontrolpoints)
-						)
-					)
-					);
-
-				// add ...	const uint vertex_id = threadgroup_position_in_grid * TessellationInputControlPoints + inputCPID;
-				ir_variable* vertex_idVar = new(ParseState)ir_variable(glsl_type::get_instance(GLSL_TYPE_UINT, 1, 1), "vertex_idVar", ir_var_temporary);
-				vf_if->then_instructions.push_tail(vertex_idVar);
-				vf_if->then_instructions.push_tail(
-					new (ParseState)ir_assignment(
-						new(ParseState)ir_dereference_variable(vertex_idVar),
-						new(ParseState)ir_expression(ir_binop_add,
-							new(ParseState)ir_expression(ir_binop_mul,
-								new(ParseState)ir_dereference_variable(internalPatchIDVar),
-								new(ParseState)ir_constant((unsigned)inputcontrolpoints)
-							),
-							new(ParseState)ir_dereference_variable(inputCPIDVar)
-						)
-					)
-					);
-
-				check(0);
-				vf_if->then_instructions.append_list(&VertexPreCallInstructions);
-
-				// call VertexMain
-				vf_if->then_instructions.push_tail(new(ParseState) ir_call(VertexEntryPointSig, NULL, &VertexArgInstructions));
-
-				// assign OutputVertexVar to InputPatchVar[patchIDInThreadgroup][inputCPID]
-				ir_dereference_array* InputPatchElementIndex = new(ParseState)ir_dereference_array(
-					new(ParseState)ir_dereference_array(
-						InputPatchVar,
-						new (ParseState)ir_dereference_variable(ParseState->symbols->get_variable("patchIDInThreadgroup"))
-					),
-					new (ParseState)ir_dereference_variable(inputCPIDVar)
-					);
-				vf_if->then_instructions.push_tail(
-					new (ParseState)ir_assignment(
-						InputPatchElementIndex,
-						new (ParseState)ir_dereference_variable(OutputVertexVar)
-					)
-					);
-			}
 
 			// call barrier() to ensure that all threads have computed the per-input-patch computation
 			{
@@ -5976,7 +5835,7 @@ bool FMetalCodeBackend::GenerateMain(EHlslShaderFrequency Frequency, const char*
 					);
 
 				ir_dereference* deref = nullptr;
-				if(!EXEC_AT_INPUT_CP_RATE || inputcontrolpoints == ParseState->tessellation.outputcontrolpoints)
+				if(inputcontrolpoints == ParseState->tessellation.outputcontrolpoints)
 				{
 					deref = EntryPointReturn;
 				}
@@ -6110,7 +5969,7 @@ bool FMetalCodeBackend::GenerateMain(EHlslShaderFrequency Frequency, const char*
 			pv_if->then_instructions.append_list(&PostMainHullCallInstructions);
 			
 			DeclInstructions.append_list(&MainHullDeclInstructions);
-			if(!EXEC_AT_INPUT_CP_RATE || inputcontrolpoints == ParseState->tessellation.outputcontrolpoints)
+			if(inputcontrolpoints == ParseState->tessellation.outputcontrolpoints)
 			{
 				PostPostCallInstructions.push_tail(pv_if);
 			}
@@ -6185,7 +6044,6 @@ bool FMetalCodeBackend::GenerateMain(EHlslShaderFrequency Frequency, const char*
 	MainSig->is_defined = true;
 	MainSig->is_main = true;
 	MainSig->body.append_list(&PrePreCallInstructions);
-#if EXEC_AT_INPUT_CP_RATE
 	if(!bIsTessellationVSHS)
 	{
 		MainSig->body.append_list(&PreCallInstructions);
@@ -6269,12 +6127,6 @@ bool FMetalCodeBackend::GenerateMain(EHlslShaderFrequency Frequency, const char*
 			vf_if->then_instructions.push_tail(pv_if);
 		}
 	}
-#else // EXEC_AT_INPUT_CP_RATE
-	MainSig->body.append_list(&PreCallInstructions);
-	// Call the original EntryPoint
-	MainSig->body.push_tail(new(ParseState) ir_call(EntryPointSig, EntryPointReturn, &ArgInstructions));
-	MainSig->body.append_list(&PostCallInstructions);
-#endif // EXEC_AT_INPUT_CP_RATE
 	MainSig->body.append_list(&PostPostCallInstructions);
 	MainSig->wg_size_x = EntryPointSig->wg_size_x;
 	MainSig->wg_size_y = EntryPointSig->wg_size_y;
@@ -6314,11 +6166,7 @@ void FMetalCodeBackend::CallPatchConstantFunction(_mesa_glsl_parse_state* ParseS
 		0u
 		),
 		new (ParseState)ir_dereference_variable(
-#if EXEC_AT_INPUT_CP_RATE
 		ParseState->symbols->get_variable("SV_InputControlPointID")
-#else // EXEC_AT_INPUT_CP_RATE
-		ParseState->symbols->get_variable("SV_OutputControlPointID")
-#endif // EXEC_AT_INPUT_CP_RATE
 		)
 		)
 		);
diff --git a/Engine/Source/Developer/Apple/MetalShaderFormat/Private/MetalBackend.h b/Engine/Source/Developer/Apple/MetalShaderFormat/Private/MetalBackend.h
--- a/Engine/Source/Developer/Apple/MetalShaderFormat/Private/MetalBackend.h
+++ b/Engine/Source/Developer/Apple/MetalShaderFormat/Private/MetalBackend.h
@@ -4,8 +4,6 @@
 
 #include "CoreMinimal.h"
 
-#define USE_VS_HS_ATTRIBUTES 1
-
 #include "hlslcc.h"
 THIRD_PARTY_INCLUDES_START
 	#include "ir.h"
diff --git a/Engine/Source/Developer/Apple/MetalShaderFormat/Private/MetalUtils.cpp b/Engine/Source/Developer/Apple/MetalShaderFormat/Private/MetalUtils.cpp
--- a/Engine/Source/Developer/Apple/MetalShaderFormat/Private/MetalUtils.cpp
+++ b/Engine/Source/Developer/Apple/MetalShaderFormat/Private/MetalUtils.cpp
@@ -1826,8 +1826,6 @@ static ir_dereference_variable* GenerateShaderOutput(
 	return TempVariableDeref;
 }
 
-#define USE_DS_ATTRIBUTES 0 // @todo Unicorn: consider using this as it would be more elegant, but currently isn't complete.
-
 void FMetalCodeBackend::build_iab_fields(_mesa_glsl_parse_state* ParseState, char const* n, struct glsl_type const* t, TArray<struct glsl_struct_field>& Fields, unsigned& FieldIndex, unsigned& BufferIndex, bool top, FBuffers const& Buffers)
 {
 	switch(t->base_type)
@@ -2110,11 +2108,8 @@ void FMetalCodeBackend::PackInputsAndOutputs(exec_list* Instructions, _mesa_glsl
 	ir_variable* DSStageIn = nullptr;
 	TIRVarList DSInputArguments;
 
-#if USE_DS_ATTRIBUTES
-#else // USE_DS_ATTRIBUTES
 	TIRVarSet DSPatchVariables;
 	ir_variable* DSPatch = nullptr;
-#endif // USE_DS_ATTRIBUTES
 
 	ir_variable* internalPatchIDVar = nullptr;
 
@@ -2163,14 +2158,9 @@ void FMetalCodeBackend::PackInputsAndOutputs(exec_list* Instructions, _mesa_glsl
 					break;
 
 				case ir_var_in:
-#if USE_VS_HS_ATTRIBUTES
 					{
 						// do nothing here MovePackedUniformsToMain will move all these input arguments to the function signature
 					}
-#else // USE_VS_HS_ATTRIBUTES
-					// there should be no input attributes
-					check(0);
-#endif // USE_VS_HS_ATTRIBUTES
 					break;
 				}
 			}
@@ -2459,10 +2449,7 @@ void FMetalCodeBackend::PackInputsAndOutputs(exec_list* Instructions, _mesa_glsl
 	{
 		// Vertex Fetch to Vertex connector
 		TArray<glsl_struct_field> DSStageInMembers;
-#if USE_DS_ATTRIBUTES
-#else // USE_DS_ATTRIBUTES
 		TArray<glsl_struct_field> DSPatchMembers;
-#endif // USE_DS_ATTRIBUTES
 
 		// Vertex Output connector. Gather position semantic & other outputs into a struct
 		TArray<glsl_struct_field> DSOutMembers;
@@ -2491,8 +2478,6 @@ void FMetalCodeBackend::PackInputsAndOutputs(exec_list* Instructions, _mesa_glsl
 					break;
 
 				case ir_var_in:
-#if USE_DS_ATTRIBUTES
-#else // USE_DS_ATTRIBUTES
 					if(Variable->type->is_patch())
 					{
 						glsl_struct_field Member;
@@ -2507,7 +2492,6 @@ void FMetalCodeBackend::PackInputsAndOutputs(exec_list* Instructions, _mesa_glsl
 						DSPatchVariables.insert(Variable);
 					}
 					else
-#endif // USE_DS_ATTRIBUTES
 					if (!ProcessStageInVariables(ParseState, bIsDesktop, Frequency, Variable, DSStageInMembers, DSStageInVariables, nullptr, DSInputArguments))
 					{
 						return;
@@ -2517,38 +2501,6 @@ void FMetalCodeBackend::PackInputsAndOutputs(exec_list* Instructions, _mesa_glsl
 			}
 		}
 
-#if USE_DS_ATTRIBUTES
-#error add attribute(0..N) semantics...
-		if (DSStageInMembers.Num())
-		{
-			check(Frequency == HSF_DomainShader);
-
-			//@todo-rco: Make me nice...
-			for (auto& Member : DSStageInMembers)
-			{
-				InputVars.push_tail(new(ParseState)extern_var(new(ParseState)ir_variable(Member.type, ralloc_strdup(ParseState, Member.name), ir_var_in)));
-			}
-
-			DSStageInMembers.Sort(
-				[](glsl_struct_field const& A, glsl_struct_field const& B)
-				{
-					return GetAttributeIndex(A.semantic) < GetAttributeIndex(B.semantic);
-				});
-
-			auto* Type = glsl_type::get_record_instance(&DSStageInMembers[0], (unsigned int)DSStageInMembers.Num(), "FDSStageIn");
-			DSStageIn = new(ParseState)ir_variable(Type, "__DSStageIn", ir_var_in);
-			// Hack: This way we tell we need to convert types from half to float
-			((glsl_type*)Type)->HlslName = "__STAGE_IN__";
-			ParseState->symbols->add_variable(DSStageIn);
-
-			if (!ParseState->AddUserStruct(Type))
-			{
-				YYLTYPE loc = {0};
-				_mesa_glsl_error(&loc, ParseState, "struct '%s' previously defined", Type->name);
-			}
-		}
-#else // USE_DS_ATTRIBUTES
-
 		// track attribute#s
 		int onAttribute = 0;
 
@@ -2736,7 +2688,6 @@ void FMetalCodeBackend::PackInputsAndOutputs(exec_list* Instructions, _mesa_glsl
 				}
 			}
 		}
-#endif // USE_DS_ATTRIBUTES
 
 		if (DSOutMembers.Num())
 		{
@@ -2804,49 +2755,9 @@ void FMetalCodeBackend::PackInputsAndOutputs(exec_list* Instructions, _mesa_glsl
 				}
 				else if(bIsTessellationVSHS)
 				{
-#if USE_VS_HS_ATTRIBUTES
 					// see above...
 					// do nothing here MovePackedUniformsToMain will move all these output arguments to the function signature
-#else // USE_VS_HS_ATTRIBUTES
-					// there should be no input attributes
-					check(0);
-#endif // USE_VS_HS_ATTRIBUTES
-				}
-#if USE_DS_ATTRIBUTES
-				else if (DSStageInVariables.find(Variable) != DSStageInVariables.end())
-				{
-					//if(Variable->type->is_outputpatch()) // not exactly what I want...
-					//{
-					//	break; // skip for the outputpatch
-					//}
-					ir_dereference* DeRefMember = new(ParseState)ir_dereference_record(DSStageIn, Variable->name);
-					if(DeRefMember->type->is_outputpatch())
-					{
-						const unsigned ElementCount = DeRefMember->type->patch_length;
-						check(Variable->type->is_outputpatch());
-						check(ElementCount == Variable->type->patch_length);
-						for (unsigned i = 0; i < ElementCount; ++i)
-						{
-							ir_dereference_array* ArrayVariable = new(ParseState)ir_dereference_array(
-								new(ParseState) ir_dereference_variable(Variable),
-								new(ParseState) ir_constant((unsigned)i)
-								);
-							ir_dereference_array* ArrayDerefMember = new(ParseState)ir_dereference_array(
-								DeRefMember->clone(ParseState, NULL),
-								new(ParseState) ir_constant((unsigned)i)
-								);
-							auto* Assign = new(ParseState)ir_assignment(ArrayVariable, ArrayDerefMember);
-							PreCallInstructions.push_tail(Assign);
-						}
-					}
-					else
-					{
-						auto* Assign = new(ParseState)ir_assignment(new(ParseState)ir_dereference_variable(Variable), DeRefMember);
-						PreCallInstructions.push_tail(Assign);
-					}
-					VarsToMoveToBody.push_back(Variable);
 				}
-#else // USE_DS_ATTRIBUTES
 				else if (DSStageInVariables.find(Variable) != DSStageInVariables.end())
 				{
 					// TOOD could merge the code above down here...
@@ -2855,7 +2766,6 @@ void FMetalCodeBackend::PackInputsAndOutputs(exec_list* Instructions, _mesa_glsl
 				{
 					// TOOD could merge the code above down here...
 				}
-#endif // USE_DS_ATTRIBUTES
 				else
 				{
 					FSemanticQualifier Qualifier;
