DR - Fix static analysis warning
#rnx

diff --git a/Engine/Source/Developer/Apple/MetalShaderFormat/Private/MetalUtils.cpp b/Engine/Source/Developer/Apple/MetalShaderFormat/Private/MetalUtils.cpp
--- a/Engine/Source/Developer/Apple/MetalShaderFormat/Private/MetalUtils.cpp
+++ b/Engine/Source/Developer/Apple/MetalShaderFormat/Private/MetalUtils.cpp
@@ -654,8 +654,6 @@ namespace MetalUtils
 					int FieldIndex = DerefStruct->type->field_index(DerefStruct->field);
 					if (FieldIndex >= 0)
 					{
-						DerefStruct->type->fields.structure[FieldIndex];
-						
 						InQualifier.Fields.bCentroid = DerefStruct->type->fields.structure[FieldIndex].centroid;
 						InQualifier.Fields.InterpolationMode = DerefStruct->type->fields.structure[FieldIndex].interpolation;
 						InQualifier.Fields.bIsPatchConstant = DerefStruct->type->fields.structure[FieldIndex].patchconstant;
@@ -668,8 +666,6 @@ namespace MetalUtils
 					int FieldIndex = DerefStruct->type->field_index(DerefStruct->field);
 					if (FieldIndex >= 0)
 					{
-						DerefStruct->type->fields.structure[FieldIndex];
-						
 						InQualifier.Fields.bCentroid = DerefStruct->type->fields.structure[FieldIndex].centroid;
 						InQualifier.Fields.InterpolationMode = DerefStruct->type->fields.structure[FieldIndex].interpolation;
 						InQualifier.Fields.bIsPatchConstant = DerefStruct->type->fields.structure[FieldIndex].patchconstant;
