Add an option "Force 32bit Floating Point Precision" to the iOS project settings to have the shader compiler promote all half-float operations to FP32 as some projects on iOS encounter rendering issues when fragment shader calculations run into half-float precision limits.

diff --git a/Engine/Config/BaseEngine.ini b/Engine/Config/BaseEngine.ini
--- a/Engine/Config/BaseEngine.ini
+++ b/Engine/Config/BaseEngine.ini
@@ -1853,6 +1853,7 @@ bEnableCloudKitSupport=False
 bGenerateCrashReportSymbols=false
 bAutomaticSigning=false
 UseFastIntrinsics=False
+ForceFloats=False
 EnableMathOptimisations=True
 MaxShaderLanguageVersion=2
 
@@ -1994,6 +1995,7 @@ TargetArchitecture=X86_64UnknownLinuxGnu
 MaxShaderLanguageVersion=3
 +TargetedRHIs=SF_METAL_SM5
 UseFastIntrinsics=False
+ForceFloats=False
 EnableMathOptimisations=True
 
 [HMDPluginPriority]
diff --git a/Engine/Source/Developer/Apple/MetalShaderFormat/Private/MetalShaderCompiler.cpp b/Engine/Source/Developer/Apple/MetalShaderFormat/Private/MetalShaderCompiler.cpp
--- a/Engine/Source/Developer/Apple/MetalShaderFormat/Private/MetalShaderCompiler.cpp
+++ b/Engine/Source/Developer/Apple/MetalShaderFormat/Private/MetalShaderCompiler.cpp
@@ -2106,20 +2106,44 @@ void CompileShader_Metal(const FShaderCompilerInput& _Input,FShaderCompilerOutpu
 			}
 			else
 			{
-				MinOSVersion = TEXT("-mmacosx-version-min=10.12");
+				Output.bSucceeded = false;
+				FShaderCompilerError* NewError = new(Output.Errors) FShaderCompilerError();
+				NewError->StrippedErrorMessage = FString::Printf(
+															 TEXT("Metal %s is no longer supported in UE4 for macOS."),
+															 *StandardVersion
+															 );
+				return;
 			}
 			break;
 		case 1:
+		{
 			HlslCompilerTarget = bIsMobile ? HlslCompilerTarget : HCT_FeatureLevelSM5;
 			StandardVersion = TEXT("1.1");
 			MinOSVersion = bIsMobile ? TEXT("") : TEXT("-mmacosx-version-min=10.11");
-			break;
+			
+			Output.bSucceeded = false;
+			FShaderCompilerError* NewError = new(Output.Errors) FShaderCompilerError();
+			NewError->StrippedErrorMessage = FString::Printf(
+														 TEXT("Metal %s is no longer supported in UE4."),
+														 *StandardVersion
+														 );
+			return;
+		}
 		case 0:
 		default:
+		{
 			check(bIsMobile);
 			StandardVersion = TEXT("1.0");
 			MinOSVersion = TEXT("");
-			break;
+			
+			Output.bSucceeded = false;
+			FShaderCompilerError* NewError = new(Output.Errors) FShaderCompilerError();
+			NewError->StrippedErrorMessage = FString::Printf(
+															 TEXT("Metal %s is no longer supported in UE4."),
+															 *StandardVersion
+															 );
+			return;
+		}
 	}
 	
 	// Force floats if the material requests it
diff --git a/Engine/Source/Developer/Mac/MacTargetPlatform/Classes/MacTargetSettings.h b/Engine/Source/Developer/Mac/MacTargetPlatform/Classes/MacTargetSettings.h
--- a/Engine/Source/Developer/Mac/MacTargetPlatform/Classes/MacTargetSettings.h
+++ b/Engine/Source/Developer/Mac/MacTargetPlatform/Classes/MacTargetSettings.h
@@ -57,6 +57,14 @@ public:
     UPROPERTY(EditAnywhere, config, Category=Rendering, meta = (DisplayName = "Use Fast-Math intrinsics", ConfigRestartRequired = true))
 	bool UseFastIntrinsics;
 	
+	/**
+	 * Whether to force Metal shaders to use 32bit floating point precision even when the shader uses half floats.
+	 * Half floats are much more efficient when they are availble but have less accuracy over large ranges,
+	 * as such some projects may need to use 32bit floats to ensure correct rendering.
+	 */
+	UPROPERTY(EditAnywhere, config, Category=Rendering, meta = (DisplayName = "Force 32bit Floating Point Precision", ConfigRestartRequired = true))
+	bool ForceFloats;
+	
 	/**
 	 * Whether to use of Metal shader-compiler's -ffast-math optimisations.
 	 * Fast-Math performs algebraic-equivalent & reassociative optimisations not permitted by the floating point arithmetic standard (IEEE-754).
diff --git a/Engine/Source/Developer/Mac/MacTargetPlatform/Private/MacTargetPlatformModule.cpp b/Engine/Source/Developer/Mac/MacTargetPlatform/Private/MacTargetPlatformModule.cpp
--- a/Engine/Source/Developer/Mac/MacTargetPlatform/Private/MacTargetPlatformModule.cpp
+++ b/Engine/Source/Developer/Mac/MacTargetPlatform/Private/MacTargetPlatformModule.cpp
@@ -78,6 +78,11 @@ public:
 			TargetSettings->EnableMathOptimisations = true;
 		}
 		
+		if (!GConfig->GetBool(TEXT("/Script/MacTargetPlatform.MacTargetSettings"), TEXT("ForceFloats"), TargetSettings->ForceFloats, GEngineIni))
+		{
+			TargetSettings->ForceFloats = false;
+		}
+		
 		TargetSettings->AddToRoot();
 
 		ISettingsModule* SettingsModule = FModuleManager::GetModulePtr<ISettingsModule>("Settings");
diff --git a/Engine/Source/Runtime/Engine/Private/ShaderCompiler/ShaderCompiler.cpp b/Engine/Source/Runtime/Engine/Private/ShaderCompiler/ShaderCompiler.cpp
--- a/Engine/Source/Runtime/Engine/Private/ShaderCompiler/ShaderCompiler.cpp
+++ b/Engine/Source/Runtime/Engine/Private/ShaderCompiler/ShaderCompiler.cpp
@@ -3145,20 +3145,24 @@ void GlobalBeginCompileShader(
 			Input.Environment.SetDefine(TEXT("MAX_SHADER_LANGUAGE_VERSION"), ShaderVersion);
 			
 			FString AllowFastIntrinsics;
+			FString ForceFloats;
 			bool bEnableMathOptimisations = true;
 			if (IsPCPlatform(EShaderPlatform(Target.Platform)))
 			{
 				GConfig->GetString(TEXT("/Script/MacTargetPlatform.MacTargetSettings"), TEXT("UseFastIntrinsics"), AllowFastIntrinsics, GEngineIni);
 				GConfig->GetBool(TEXT("/Script/MacTargetPlatform.MacTargetSettings"), TEXT("EnableMathOptimisations"), bEnableMathOptimisations, GEngineIni);
+				GConfig->GetString(TEXT("/Script/MacTargetPlatform.MacTargetSettings"), TEXT("ForceFloats"), ForceFloats, GEngineIni);
 			}
 			else
 			{
 				GConfig->GetString(TEXT("/Script/IOSRuntimeSettings.IOSRuntimeSettings"), TEXT("UseFastIntrinsics"), AllowFastIntrinsics, GEngineIni);
 				GConfig->GetBool(TEXT("/Script/IOSRuntimeSettings.IOSRuntimeSettings"), TEXT("EnableMathOptimisations"), bEnableMathOptimisations, GEngineIni);
+				GConfig->GetString(TEXT("/Script/IOSRuntimeSettings.IOSRuntimeSettings"), TEXT("ForceFloats"), ForceFloats, GEngineIni);
 				// Force no development shaders on iOS
 				bAllowDevelopmentShaderCompile = false;
 			}
 			Input.Environment.SetDefine(TEXT("METAL_USE_FAST_INTRINSICS"), *AllowFastIntrinsics);
+			Input.Environment.SetDefine(TEXT("FORCE_FLOATS"), *ForceFloats);
 			
 			// Same as console-variable above, but that's global and this is per-platform, per-project
 			if (!bEnableMathOptimisations)
diff --git a/Engine/Source/Runtime/IOS/IOSRuntimeSettings/Classes/IOSRuntimeSettings.h b/Engine/Source/Runtime/IOS/IOSRuntimeSettings/Classes/IOSRuntimeSettings.h
--- a/Engine/Source/Runtime/IOS/IOSRuntimeSettings/Classes/IOSRuntimeSettings.h
+++ b/Engine/Source/Runtime/IOS/IOSRuntimeSettings/Classes/IOSRuntimeSettings.h
@@ -429,6 +429,14 @@ public:
 	UPROPERTY(EditAnywhere, config, Category=Rendering, meta = (DisplayName = "Use Fast-Math intrinsics", ConfigRestartRequired = true))
 	bool UseFastIntrinsics;
 	
+	/**
+	 * Whether to force Metal shaders to use 32bit floating point precision even when the shader uses half floats.
+	 * Half floats are much more efficient when they are availble but have less accuracy over large ranges,
+	 * as such some projects may need to use 32bit floats to ensure correct rendering.
+	 */
+	UPROPERTY(EditAnywhere, config, Category=Rendering, meta = (DisplayName = "Force 32bit Floating Point Precision", ConfigRestartRequired = true))
+	bool ForceFloats;
+	
 	/**
 	 * Whether to use of Metal shader-compiler's -ffast-math optimisations.
 	 * Fast-Math performs algebraic-equivalent & reassociative optimisations not permitted by the floating point arithmetic standard (IEEE-754).
diff --git a/Engine/Source/Runtime/RenderCore/Private/Shader.cpp b/Engine/Source/Runtime/RenderCore/Private/Shader.cpp
--- a/Engine/Source/Runtime/RenderCore/Private/Shader.cpp
+++ b/Engine/Source/Runtime/RenderCore/Private/Shader.cpp
@@ -2227,15 +2227,18 @@ void ShaderMapAppendKeyString(EShaderPlatform Platform, FString& KeyString)
 		
 		bool bAllowFastIntrinsics = false;
 		bool bEnableMathOptimisations = true;
+		bool bForceFloats = false;
 		if (IsPCPlatform(Platform))
 		{
 			GConfig->GetBool(TEXT("/Script/MacTargetPlatform.MacTargetSettings"), TEXT("UseFastIntrinsics"), bAllowFastIntrinsics, GEngineIni);
 			GConfig->GetBool(TEXT("/Script/MacTargetPlatform.MacTargetSettings"), TEXT("EnableMathOptimisations"), bEnableMathOptimisations, GEngineIni);
+			GConfig->GetBool(TEXT("/Script/MacTargetPlatform.MacTargetSettings"), TEXT("ForceFloats"), bForceFloats, GEngineIni);
 		}
 		else
 		{
 			GConfig->GetBool(TEXT("/Script/IOSRuntimeSettings.IOSRuntimeSettings"), TEXT("UseFastIntrinsics"), bAllowFastIntrinsics, GEngineIni);
 			GConfig->GetBool(TEXT("/Script/IOSRuntimeSettings.IOSRuntimeSettings"), TEXT("EnableMathOptimisations"), bEnableMathOptimisations, GEngineIni);
+			GConfig->GetBool(TEXT("/Script/IOSRuntimeSettings.IOSRuntimeSettings"), TEXT("ForceFloats"), bForceFloats, GEngineIni);
 		}
 		
 		if (bAllowFastIntrinsics)
@@ -2249,6 +2252,11 @@ void ShaderMapAppendKeyString(EShaderPlatform Platform, FString& KeyString)
 			KeyString += TEXT("_NoFastMath");
 		}
 		
+		if (bForceFloats)
+		{
+			KeyString += TEXT("_FP32");
+		}
+		
 		// Shaders built for archiving - for Metal that requires compiling the code in a different way so that we can strip it later
 		bool bArchive = false;
 		GConfig->GetBool(TEXT("/Script/UnrealEd.ProjectPackagingSettings"), TEXT("bSharedMaterialNativeLibraries"), bArchive, GGameIni);
