Support packed vector types in RW/StructureBuffer types in MetalSL.

diff --git a/Engine/Source/Developer/Apple/MetalShaderFormat/Private/MetalBackend.cpp b/Engine/Source/Developer/Apple/MetalShaderFormat/Private/MetalBackend.cpp
--- a/Engine/Source/Developer/Apple/MetalShaderFormat/Private/MetalBackend.cpp
+++ b/Engine/Source/Developer/Apple/MetalShaderFormat/Private/MetalBackend.cpp
@@ -1032,6 +1032,14 @@ protected:
 					bool bIsAtomic = Buffers.AtomicVariables.find(var) != Buffers.AtomicVariables.end();
                     if (bIsStructuredBuffer || bIsByteAddressBuffer || bIsInvariantType || bIsAtomic)
 					{
+						if (var->type->inner_type->is_record())
+						{
+							if (hash_table_find(used_structures, var->type->inner_type) == NULL)
+							{
+								hash_table_insert(used_structures, (void*)var->type->inner_type, var->type->inner_type);
+							}
+						}
+						
 						check(BufferIndex <= 30);
                         
                         if (Buffers.AtomicVariables.find(var) == Buffers.AtomicVariables.end())
@@ -1140,6 +1148,14 @@ protected:
 							bool bIsAtomic = Buffers.AtomicVariables.find(var) != Buffers.AtomicVariables.end();
 							if (bIsStructuredBuffer || bIsByteAddressBuffer || bIsInvariantType || bIsAtomic)
 							{
+								if (var->type->inner_type->is_record())
+								{
+									if (hash_table_find(used_structures, var->type->inner_type) == NULL)
+									{
+										hash_table_insert(used_structures, (void*)var->type->inner_type, var->type->inner_type);
+									}
+								}
+								
 								check(BufferIndex >= 0 && BufferIndex <= 30);
 								ralloc_asprintf_append(
 													   buffer,
@@ -3044,18 +3060,27 @@ protected:
 		{
 			char mask[6];
 			unsigned j = 1;
+			
 			if (assign->lhs->type->is_scalar() == false ||
 				assign->write_mask != 0x1)
 			{
-				for (unsigned i = 0; i < 4; i++)
+				auto* DerefRecord = assign->lhs->as_dereference_record();
+				bool bPackableRecord = (assign->lhs->as_dereference_record() && DerefRecord->record->type->HlslName && !strcmp(DerefRecord->record->type->HlslName, "__PACKED__"));
+				bool bPackableVector = assign->lhs->type->is_vector() && assign->lhs->type->vector_elements < 4;
+				if (!bPackableRecord || !bPackableVector)
+				
 				{
-					if ((assign->write_mask & (1 << i)) != 0)
+					for (unsigned i = 0; i < 4; i++)
 					{
-						mask[j] = "xyzw"[i];
-						j++;
+						if ((assign->write_mask & (1 << i)) != 0)
+						{
+							mask[j] = "xyzw"[i];
+							j++;
+						}
 					}
 				}
 			}
+			
 			mask[j] = '\0';
 
 			mask[0] = (j == 1) ? '\0' : '.';
diff --git a/Engine/Source/Developer/Apple/MetalShaderFormat/Private/MetalUtils.cpp b/Engine/Source/Developer/Apple/MetalShaderFormat/Private/MetalUtils.cpp
--- a/Engine/Source/Developer/Apple/MetalShaderFormat/Private/MetalUtils.cpp
+++ b/Engine/Source/Developer/Apple/MetalShaderFormat/Private/MetalUtils.cpp
@@ -1251,6 +1251,12 @@ void FMetalCodeBackend::MovePackedUniformsToMain(exec_list* ir, _mesa_glsl_parse
 			{
 				if (bIsBuffer)
 				{
+                    bool bReallyIsStructuredBuffer = Var->type->sampler_buffer && Var->type->inner_type && Var->type->inner_type->is_record() && (!strncmp(Var->type->name, "RWStructuredBuffer<", 19) || !strncmp(Var->type->name, "StructuredBuffer<", 17));
+                    if (bReallyIsStructuredBuffer)
+					{
+						((glsl_type*)Var->type->inner_type)->HlslName = "__PACKED__";
+					}
+					
 					OutBuffers.AddBuffer(Var);
 				}
 				else
@@ -3295,6 +3301,20 @@ struct FDeReferencePackedVarsVisitor final : public ir_rvalue_visitor
 	virtual ~FDeReferencePackedVarsVisitor()
 	{
 	}
+	
+	virtual ir_visitor_status visit_leave(ir_assignment * ir) override
+	{
+		auto* DerefRecord = ir->lhs->as_dereference_record();
+		if (DerefRecord)
+		{
+            auto* StructVar = ir->lhs->variable_referenced();
+            if (StructVar->type->base_type == GLSL_TYPE_IMAGE && StructVar->type->inner_type->is_record() && StructVar->type->inner_type->HlslName && !strcmp(StructVar->type->inner_type->HlslName, "__PACKED__"))
+            {
+                handle_rvalue((ir_rvalue**)&ir->lhs);
+            }
+		}
+		return ir_rvalue_visitor::visit_leave(ir);
+	}
 
 	virtual ir_visitor_status visit_enter(ir_expression* ir) override
 	{
@@ -3317,6 +3337,16 @@ struct FDeReferencePackedVarsVisitor final : public ir_rvalue_visitor
 		--ExpressionDepth;
 		return ir_rvalue_visitor::visit_leave(ir);
 	}
+	
+	virtual ir_visitor_status visit_leave(ir_dereference_record * ir) override
+	{
+        auto* StructVar = ir->variable_referenced();
+        if (StructVar->type->base_type == GLSL_TYPE_IMAGE && StructVar->type->inner_type->is_record() && StructVar->type->inner_type->HlslName && !strcmp(StructVar->type->inner_type->HlslName, "__PACKED__"))
+        {
+            handle_rvalue((ir_rvalue**)&ir);
+        }
+		return ir_rvalue_visitor::visit_leave(ir);
+	}
 
 	virtual void handle_rvalue(ir_rvalue** RValuePtr) override
 	{
@@ -3340,17 +3370,25 @@ struct FDeReferencePackedVarsVisitor final : public ir_rvalue_visitor
 				}
 			}
 		}
-		else if (DeRefRecord && ExpressionDepth > 0)
+		else if (DeRefRecord)
 		{
 			auto* StructVar = DeRefRecord->variable_referenced();
-			if (StructVar->type->HlslName && !strcmp(StructVar->type->HlslName, "__PACKED__"))
+			if (ExpressionDepth > 0 && StructVar->type->HlslName && !strcmp(StructVar->type->HlslName, "__PACKED__"))
 			{
 				if (DeRefRecord->type->vector_elements > 1 && DeRefRecord->type->vector_elements < 4)
 				{
 					auto* Var = GetVar(DeRefRecord);
 					*RValuePtr = new(State)ir_dereference_variable(Var);
 				}
-			}
+            }
+            else if (StructVar->type->base_type == GLSL_TYPE_IMAGE && StructVar->type->inner_type->is_record() && StructVar->type->inner_type->HlslName && !strcmp(StructVar->type->inner_type->HlslName, "__PACKED__"))
+            {
+                if (DeRefRecord->type->vector_elements > 1 && DeRefRecord->type->vector_elements < 4)
+                {
+                    auto* Var = GetVar(DeRefRecord);
+                    *RValuePtr = new(State)ir_dereference_variable(Var);
+                }
+            }
 		}
 	}
 
