Dupe 4858714:
Remove obsolete alternative Metal shader frequency table that was causing PVS to spit out an invalid error.

diff --git a/Engine/Source/Developer/Apple/MetalShaderFormat/Private/MetalShaderCompiler.cpp b/Engine/Source/Developer/Apple/MetalShaderFormat/Private/MetalShaderCompiler.cpp
--- a/Engine/Source/Developer/Apple/MetalShaderFormat/Private/MetalShaderCompiler.cpp
+++ b/Engine/Source/Developer/Apple/MetalShaderFormat/Private/MetalShaderCompiler.cpp
@@ -1846,19 +1846,7 @@ void BuildMetalShaderOutput(
 	External interface.
 ------------------------------------------------------------------------------*/
 
-// For Metal <= 1.1
-static const EHlslShaderFrequency FrequencyTable1[] =
-{
-	HSF_VertexShader,
-	HSF_InvalidFrequency,
-	HSF_InvalidFrequency,
-	HSF_PixelShader,
-	HSF_InvalidFrequency,
-	HSF_ComputeShader
-};
-
-// For Metal >= 1.2
-static const EHlslShaderFrequency FrequencyTable2[] =
+static const EHlslShaderFrequency FrequencyTable[] =
 {
 	HSF_VertexShader,
 	HSF_HullShader,
@@ -2232,7 +2220,7 @@ void CompileShader_Metal(const FShaderCompilerInput& _Input,FShaderCompilerOutpu
 	char* MetalShaderSource = NULL;
 	char* ErrorLog = NULL;
 
-	const EHlslShaderFrequency Frequency = HlslCompilerTarget < HCT_FeatureLevelSM5 ? FrequencyTable1[Input.Target.Frequency] : FrequencyTable2[Input.Target.Frequency];
+	const EHlslShaderFrequency Frequency = FrequencyTable[Input.Target.Frequency];
 	if (Frequency == HSF_InvalidFrequency)
 	{
 		Output.bSucceeded = false;
