Make RWBuffer's use raw Metal T* buffers so that we don't run into problems where texture_buffer/texture2d can't be both read & write at the same time on iOS. This will be slower but it will work on more iOS/tvOS devices. This does not affect macOS.

diff --git a/Engine/Source/Developer/Apple/MetalShaderFormat/Private/MetalBackend.cpp b/Engine/Source/Developer/Apple/MetalShaderFormat/Private/MetalBackend.cpp
--- a/Engine/Source/Developer/Apple/MetalShaderFormat/Private/MetalBackend.cpp
+++ b/Engine/Source/Developer/Apple/MetalShaderFormat/Private/MetalBackend.cpp
@@ -1464,7 +1464,7 @@ protected:
 				{
 					bool bIsStructuredBuffer = (inst->type->inner_type->is_record() || !strncmp(inst->type->name, "RWStructuredBuffer<", 19) || !strncmp(inst->type->name, "StructuredBuffer<", 17));
 					bool bIsByteAddressBuffer = (!strncmp(inst->type->name, "RWByteAddressBuffer", 19) || !strncmp(inst->type->name, "ByteAddressBuffer", 17));
-                	if (Buffers.AtomicVariables.find(inst) != Buffers.AtomicVariables.end() || bIsStructuredBuffer || bIsByteAddressBuffer || inst->invariant || (inst->type->components() == 3) || inst->type->inner_type->components() == 3 || Backend.Version <= 2)
+                	if (Buffers.AtomicVariables.find(inst) != Buffers.AtomicVariables.end() || bIsStructuredBuffer || bIsByteAddressBuffer || inst->invariant || ((inst->type->components() == 3) || (Backend.TypedMode == EMetalTypeBufferMode2DSRV || Backend.TypedMode == EMetalTypeBufferModeTBSRV) && inst->type->is_image()) || inst->type->inner_type->components() == 3 || Backend.Version <= 2)
 					{
 						bInsertSideTable |= true;
 					}
@@ -4747,10 +4747,18 @@ public:
                 ralloc_asprintf_append(buffer, "#define __METAL_TYPED_BUFFER_READ_IMPL__ 0\n");
                 ralloc_asprintf_append(buffer, "#define __METAL_TYPED_BUFFER_RW_IMPL__ 0\n");
                 break;
+			case EMetalTypeBufferMode2DSRV:
+				ralloc_asprintf_append(buffer, "#define __METAL_TYPED_BUFFER_READ_IMPL__ 1\n");
+				ralloc_asprintf_append(buffer, "#define __METAL_TYPED_BUFFER_RW_IMPL__ 0\n");
+				break;
             case EMetalTypeBufferMode2D:
                 ralloc_asprintf_append(buffer, "#define __METAL_TYPED_BUFFER_READ_IMPL__ 1\n");
                 ralloc_asprintf_append(buffer, "#define __METAL_TYPED_BUFFER_RW_IMPL__ 1\n");
                 break;
+			case EMetalTypeBufferModeTBSRV:
+				ralloc_asprintf_append(buffer, "#define __METAL_TYPED_BUFFER_READ_IMPL__ 3\n");
+				ralloc_asprintf_append(buffer, "#define __METAL_TYPED_BUFFER_RW_IMPL__ 0\n");
+				break;
             case EMetalTypeBufferModeTB:
                 ralloc_asprintf_append(buffer, "#define __METAL_TYPED_BUFFER_READ_IMPL__ 3\n");
                 ralloc_asprintf_append(buffer, "#define __METAL_TYPED_BUFFER_RW_IMPL__ 3\n");
diff --git a/Engine/Source/Developer/Apple/MetalShaderFormat/Private/MetalBackend.h b/Engine/Source/Developer/Apple/MetalShaderFormat/Private/MetalBackend.h
--- a/Engine/Source/Developer/Apple/MetalShaderFormat/Private/MetalBackend.h
+++ b/Engine/Source/Developer/Apple/MetalShaderFormat/Private/MetalBackend.h
@@ -76,8 +76,10 @@ enum EMetalGPUSemantics
 enum EMetalTypeBufferMode
 {
 	EMetalTypeBufferModeRaw = 0, // No typed buffers
-    EMetalTypeBufferMode2D = 2, // Buffer<> SRVs & RWBuffer<> UAVs are typed via 2D textures
-    EMetalTypeBufferModeTB = 3, // Buffer<> SRVs & RWBuffer<> UAVs are typed via texture-buffers
+	EMetalTypeBufferMode2DSRV = 1, // Buffer<> SRVs are typed via 2D textures, RWBuffer<> UAVs are raw buffers
+	EMetalTypeBufferModeTBSRV = 2, // Buffer<> SRVs are typed via texture-buffers, RWBuffer<> UAVs are raw buffers
+    EMetalTypeBufferMode2D = 3, // Buffer<> SRVs & RWBuffer<> UAVs are typed via 2D textures
+    EMetalTypeBufferModeTB = 4, // Buffer<> SRVs & RWBuffer<> UAVs are typed via texture-buffers
 };
 
 // Metal supports 16 across all HW
diff --git a/Engine/Source/Developer/Apple/MetalShaderFormat/Private/MetalShaderCompiler.cpp b/Engine/Source/Developer/Apple/MetalShaderFormat/Private/MetalShaderCompiler.cpp
--- a/Engine/Source/Developer/Apple/MetalShaderFormat/Private/MetalShaderCompiler.cpp
+++ b/Engine/Source/Developer/Apple/MetalShaderFormat/Private/MetalShaderCompiler.cpp
@@ -2058,10 +2058,12 @@ void CompileShader_Metal(const FShaderCompilerInput& _Input,FShaderCompilerOutpu
 			if (bAppleTV)
 			{
 				MinOSVersion = TEXT("-mtvos-version-min=12.0");
+				TypeMode = EMetalTypeBufferModeTBSRV;
 			}
 			else if (bIsMobile)
 			{
 				MinOSVersion = TEXT("-mios-version-min=12.0");
+				TypeMode = EMetalTypeBufferModeTBSRV;
 			}
 			else
 			{
@@ -2076,10 +2078,12 @@ void CompileShader_Metal(const FShaderCompilerInput& _Input,FShaderCompilerOutpu
 			if (bAppleTV)
 			{
 				MinOSVersion = TEXT("-mtvos-version-min=11.0");
+				TypeMode = EMetalTypeBufferMode2DSRV;
 			}
 			else if (bIsMobile)
 			{
 				MinOSVersion = TEXT("-mios-version-min=11.0");
+				TypeMode = EMetalTypeBufferMode2DSRV;
 			}
 			else
 			{
@@ -2094,10 +2098,12 @@ void CompileShader_Metal(const FShaderCompilerInput& _Input,FShaderCompilerOutpu
 			if (bAppleTV)
 			{
 				MinOSVersion = TEXT("-mtvos-version-min=10.0");
+				TypeMode = EMetalTypeBufferMode2DSRV;
 			}
 			else if (bIsMobile)
 			{
 				MinOSVersion = TEXT("-mios-version-min=10.0");
+				TypeMode = EMetalTypeBufferMode2DSRV;
 			}
 			else
 			{
diff --git a/Engine/Source/Developer/Apple/MetalShaderFormat/Private/MetalUtils.cpp b/Engine/Source/Developer/Apple/MetalShaderFormat/Private/MetalUtils.cpp
--- a/Engine/Source/Developer/Apple/MetalShaderFormat/Private/MetalUtils.cpp
+++ b/Engine/Source/Developer/Apple/MetalShaderFormat/Private/MetalUtils.cpp
@@ -1253,6 +1253,12 @@ void FMetalCodeBackend::MovePackedUniformsToMain(exec_list* ir, _mesa_glsl_parse
 					bIsBuffer = (!Var->type->is_sampler() && !Var->type->is_image()) || Var->type->sampler_buffer;
 					break;
 				}
+				case EMetalTypeBufferMode2DSRV:
+				case EMetalTypeBufferModeTBSRV:
+				{
+					bIsBuffer = (!Var->type->is_sampler() && !Var->type->is_image()) || (Var->type->sampler_buffer && (Var->type->is_image() || OutBuffers.AtomicVariables.find(Var) != OutBuffers.AtomicVariables.end() || bIsStructuredBuffer || bIsInvariant || bIsByteAddressBuffer)) || bIsVec3;
+					break;
+				}
 				case EMetalTypeBufferMode2D:
 				case EMetalTypeBufferModeTB:
 				{
