Missed some Metal shader defines from CL #5419941 that were then causing compile errors on the builders.

diff --git a/Engine/Source/Developer/Apple/MetalShaderFormat/Private/MetalShaderCompiler.cpp b/Engine/Source/Developer/Apple/MetalShaderFormat/Private/MetalShaderCompiler.cpp
--- a/Engine/Source/Developer/Apple/MetalShaderFormat/Private/MetalShaderCompiler.cpp
+++ b/Engine/Source/Developer/Apple/MetalShaderFormat/Private/MetalShaderCompiler.cpp
@@ -1679,6 +1679,14 @@ void BuildMetalShaderOutput(
 						Defines += TEXT(" -D__METAL_TYPED_BUFFER_READ_IMPL__=0");
 						Defines += TEXT(" -D__METAL_TYPED_BUFFER_RW_IMPL__=0");
 						break;
+					case EMetalTypeBufferMode2DSRV:
+						Defines += TEXT(" -D__METAL_TYPED_BUFFER_READ_IMPL__=1");
+						Defines += TEXT(" -D__METAL_TYPED_BUFFER_RW_IMPL__=0");
+						break;
+					case EMetalTypeBufferModeTBSRV:
+						Defines += TEXT(" -D__METAL_TYPED_BUFFER_READ_IMPL__=3");
+						Defines += TEXT(" -D__METAL_TYPED_BUFFER_RW_IMPL__=0");
+						break;
 					case EMetalTypeBufferMode2D:
 						Defines += TEXT(" -D__METAL_TYPED_BUFFER_READ_IMPL__=1");
 						Defines += TEXT(" -D__METAL_TYPED_BUFFER_RW_IMPL__=1");
