Implement support for arbitrary texture swizzling in Metal shaders based on a swizzle provided from the CPU - this allows handling of some formats that Metal doesn't but it has significant cost so should only be enabled when strictly necessary and this must be known at shader compile time. This is a hedge against future need.

diff --git a/Engine/Source/Developer/Apple/MetalShaderFormat/Private/MetalBackend.cpp b/Engine/Source/Developer/Apple/MetalShaderFormat/Private/MetalBackend.cpp
--- a/Engine/Source/Developer/Apple/MetalShaderFormat/Private/MetalBackend.cpp
+++ b/Engine/Source/Developer/Apple/MetalShaderFormat/Private/MetalBackend.cpp
@@ -235,112 +235,6 @@ static const char * const MetalExpressionTable[ir_opcode_count][4] =
 
 static_assert((sizeof(MetalExpressionTable) / sizeof(MetalExpressionTable[0])) == ir_opcode_count, "Metal Expression Table Size Mismatch");
 
-struct SDMARange
-{
-	unsigned SourceCB;
-	unsigned SourceOffset;
-	unsigned Size;
-	unsigned DestCBIndex;
-	unsigned DestCBPrecision;
-	unsigned DestOffset;
-
-	bool operator <(SDMARange const & Other) const
-	{
-		if (SourceCB == Other.SourceCB)
-		{
-			return SourceOffset < Other.SourceOffset;
-		}
-
-		return SourceCB < Other.SourceCB;
-	}
-};
-typedef std::list<SDMARange> TDMARangeList;
-typedef std::map<unsigned, TDMARangeList> TCBDMARangeMap;
-
-
-static void InsertRange( TCBDMARangeMap& CBAllRanges, unsigned SourceCB, unsigned SourceOffset, unsigned Size, unsigned DestCBIndex, unsigned DestCBPrecision, unsigned DestOffset ) 
-{
-	check(SourceCB < (1 << 12));
-	check(DestCBIndex < (1 << 12));
-	check(DestCBPrecision < (1 << 8));
-	unsigned SourceDestCBKey = (SourceCB << 20) | (DestCBIndex << 8) | DestCBPrecision;
-	SDMARange Range = { SourceCB, SourceOffset, Size, DestCBIndex, DestCBPrecision, DestOffset };
-
-	TDMARangeList& CBRanges = CBAllRanges[SourceDestCBKey];
-//printf("* InsertRange: %08x\t%u:%u - %u:%c:%u:%u\n", SourceDestCBKey, SourceCB, SourceOffset, DestCBIndex, DestCBPrecision, DestOffset, Size);
-	if (CBRanges.empty())
-	{
-		CBRanges.push_back(Range);
-	}
-	else
-	{
-		TDMARangeList::iterator Prev = CBRanges.end();
-		bool bAdded = false;
-		for (auto Iter = CBRanges.begin(); Iter != CBRanges.end(); ++Iter)
-		{
-			if (SourceOffset + Size <= Iter->SourceOffset)
-			{
-				if (Prev == CBRanges.end())
-				{
-					CBRanges.push_front(Range);
-				}
-				else
-				{
-					CBRanges.insert(Iter, Range);
-				}
-
-				bAdded = true;
-				break;
-			}
-
-			Prev = Iter;
-		}
-
-		if (!bAdded)
-		{
-			CBRanges.push_back(Range);
-		}
-
-		if (CBRanges.size() > 1)
-		{
-			// Try to merge ranges
-			bool bDirty = false;
-			do
-			{
-				bDirty = false;
-				TDMARangeList NewCBRanges;
-				for (auto Iter = CBRanges.begin(); Iter != CBRanges.end(); ++Iter)
-				{
-					if (Iter == CBRanges.begin())
-					{
-						Prev = CBRanges.begin();
-					}
-					else
-					{
-						if (Prev->SourceOffset + Prev->Size == Iter->SourceOffset && Prev->DestOffset + Prev->Size == Iter->DestOffset)
-						{
-							SDMARange Merged = *Prev;
-							Merged.Size = Prev->Size + Iter->Size;
-							NewCBRanges.pop_back();
-							NewCBRanges.push_back(Merged);
-							++Iter;
-							NewCBRanges.insert(NewCBRanges.end(), Iter, CBRanges.end());
-							bDirty = true;
-							break;
-						}
-					}
-
-					NewCBRanges.push_back(*Iter);
-					Prev = Iter;
-				}
-
-				CBRanges.swap(NewCBRanges);
-			}
-			while (bDirty);
-		}
-	}
-}
-
 static TDMARangeList SortRanges( TCBDMARangeMap& CBRanges ) 
 {
 	TDMARangeList Sorted;
@@ -1457,6 +1351,7 @@ protected:
         bInsertSideTable = false;
         if (sig->is_main && Backend.bBoundsChecks)
 		{
+			bInsertSideTable = Backend.bSwizzleSample;
             foreach_iter(exec_list_iterator, iter, sig->parameters)
             {
                 ir_variable *const inst = (ir_variable *) iter.get();
@@ -2021,12 +1916,20 @@ protected:
 		check(scope_depth > 0);
 		bool bNeedsClosingParenthesis = true;
 		bool bDepthTypeExpand = tex->sampler->type->sampler_shadow && !tex->shadow_comparitor;
+		bool bSwizzleSample = Backend.bSwizzleSample && (tex->sampler->type->is_sampler() && !tex->sampler->type->sampler_buffer && !tex->shadow_comparitor);
 		switch (tex->op)
 		{
+			case ir_txf:
+				if (bSwizzleSample)
+					ralloc_asprintf_append(buffer, "swizzle_sample(");
+				break;
 			case ir_tex:
 			case ir_txl:
 			case ir_txb:
 			case ir_txd:
+				if (bSwizzleSample)
+					ralloc_asprintf_append(buffer, "swizzle_sample(");
+				
 				if (bDepthTypeExpand)
 				{
 					print_type_pre(tex->type);
@@ -2039,6 +1942,43 @@ protected:
 			default:
 				break;
 		}
+		
+		auto* Texture = tex->sampler->variable_referenced();
+		check(Texture);
+		
+		int Index = 0;
+		char const* BufferSizesName = "BufferSizes";
+		bool bSideTable = bInsertSideTable;
+		{
+			if (Texture->mode == ir_var_temporary)
+			{
+				// IAB sampling path
+				ir_variable* IABVariable = Backend.IABVariablesMap.FindChecked(Texture);
+				int FieldIndex = IABVariable->type->field_index(Texture->name);
+				for (int i = 0; i < FieldIndex; i++)
+				{
+					if (IABVariable->type->fields.structure[i].type->sampler_buffer)
+					{
+						Index++;
+					}
+				}
+				
+				BufferSizesName = ralloc_asprintf(ParseState, "%s.BufferSizes", IABVariable->name);
+				bSideTable = true;
+				check(Index >= 0);
+			}
+			else
+			{
+				// Function argument path
+				Index = Buffers.GetIndex(Texture);
+				check(Index >= 0);
+				if (bSwizzleSample)
+				{
+					Index *= 2;
+					Index += 31 * 2;
+				}
+			}
+		}
 
 		bool bTexCubeArray = tex->sampler->type->sampler_array && (tex->sampler->type->sampler_dimensionality == GLSL_SAMPLER_DIM_CUBE);
 		if (tex->op != ir_txf)
@@ -2067,9 +2007,6 @@ protected:
 				ralloc_asprintf_append(buffer, ", ");
 			}
 			
-			auto* Texture = tex->sampler->variable_referenced();
-			check(Texture);
-			
 			if (tex->SamplerState)
 			{
 				tex->SamplerState->accept(this);
@@ -2205,10 +2142,19 @@ protected:
 				tex->offset->accept(this);
 			}
 			
+			ralloc_asprintf_append(buffer, ")");
+			
+			bNeedsClosingParenthesis = false;
+			
 			if (bDepthTypeExpand)
 			{
 				ralloc_asprintf_append(buffer, ")");
 			}
+			
+			if (bSwizzleSample)
+			{
+				ralloc_asprintf_append(buffer, ", %s[%d])", BufferSizesName, Index);
+			}
 		}
 			break;
 
@@ -2217,35 +2163,8 @@ protected:
 			check(tex->sampler->type);
 			if (tex->sampler->type->is_sampler() && tex->sampler->type->sampler_buffer)
 			{
-				auto* Texture = tex->sampler->variable_referenced();
-				int Index = 0;
-				char const* BufferSizesName = "BufferSizes";
-				int FormatIndex = -1;
-				bool bSideTable = bInsertSideTable;
-				if (Texture->mode == ir_var_temporary)
-				{
-					// IAB sampling path
-					ir_variable* IABVariable = Backend.IABVariablesMap.FindChecked(Texture);
-					int FieldIndex = IABVariable->type->field_index(Texture->name);
-					for (int i = 0; i < FieldIndex; i++)
-					{
-						if (IABVariable->type->fields.structure[i].type->sampler_buffer)
-						{
-							Index++;
-						}
-					}
-					
-					BufferSizesName = ralloc_asprintf(ParseState, "%s.BufferSizes", IABVariable->name);
-					bSideTable = true;
-				}
-				else
-				{
-					// Function argument path
-					Index = Buffers.GetIndex(Texture);
-					FormatIndex = Index;
-				}
-				check(Index >= 0 && Index <= 30);
-					
+				check(Index >= 0 && (Texture->mode == ir_var_temporary || Index <= 30));
+				
 				ralloc_asprintf_append(buffer, "(");
 				
 				bool bIsStructuredBuffer = (Texture->type->inner_type->is_record() || !strncmp(Texture->type->name, "RWStructuredBuffer<", 19) || !strncmp(Texture->type->name, "StructuredBuffer<", 17));
@@ -2322,6 +2241,14 @@ protected:
 					ralloc_asprintf_append(buffer, ",");
 					tex->lod_info.lod->accept(this);
 				}
+				
+				ralloc_asprintf_append(buffer, ")");
+				bNeedsClosingParenthesis = false;
+				
+				if (bSwizzleSample)
+				{
+					ralloc_asprintf_append(buffer, ", %s[%d])", BufferSizesName, Index);
+				}
 			}
 		}
 			break;
@@ -2344,10 +2271,6 @@ protected:
 				ralloc_asprintf_append(buffer, ", ");
 			}
 			
-			// Sampler
-			auto* Texture = tex->sampler->variable_referenced();
-			check(Texture);
-			
 			if (tex->SamplerState)
 			{
 				tex->SamplerState->accept(this);
@@ -2673,11 +2596,46 @@ protected:
 		
 		check( 1 <= dst_elements && dst_elements <= 4);
 		check( 1 <= src_elements && src_elements <= 4);
-
+		
 		if ( deref->op == ir_image_access)
 		{
 			bool bIsRWTexture = !deref->image->type->sampler_buffer;
 			bool bIsArray = bIsRWTexture && strstr(deref->image->type->name, "Array") != nullptr;
+			
+			auto* Texture = deref->image->variable_referenced();
+			int Index = 0;
+			char const* BufferSizesName = "BufferSizes";
+			bool bSideTable = bInsertSideTable;
+			if (Texture->mode == ir_var_temporary)
+			{
+				// IAB sampling path
+				ir_variable* IABVariable = Backend.IABVariablesMap.FindChecked(Texture);
+				int FieldIndex = IABVariable->type->field_index(Texture->name);
+				for (int i = 0; i < FieldIndex; i++)
+				{
+					if (IABVariable->type->fields.structure[i].type->sampler_buffer)
+					{
+						Index++;
+					}
+				}
+				
+				BufferSizesName = ralloc_asprintf(ParseState, "%s.BufferSizes", IABVariable->name);
+				bSideTable = true;
+				check(Index >= 0);
+			}
+			else
+			{
+				// Function argument path
+				Index = Buffers.GetIndex(Texture);
+				check(Index >= 0);
+				
+				if (bIsRWTexture)
+				{
+					Index *= 2;
+					Index += 31 * 2;
+				}
+			}
+			
 			if (src == nullptr)
 			{
 				if (bIsRWTexture)
@@ -2706,35 +2664,6 @@ protected:
 				}
 				else
 				{
-					auto* Texture = deref->image->variable_referenced();
-					int Index = 0;
-					char const* BufferSizesName = "BufferSizes";
-					int FormatIndex = -1;
-					bool bSideTable = bInsertSideTable;
-					if (Texture->mode == ir_var_temporary)
-					{
-						// IAB sampling path
-						ir_variable* IABVariable = Backend.IABVariablesMap.FindChecked(Texture);
-						int FieldIndex = IABVariable->type->field_index(Texture->name);
-						for (int i = 0; i < FieldIndex; i++)
-						{
-							if (IABVariable->type->fields.structure[i].type->sampler_buffer)
-							{
-								Index++;
-							}
-						}
-						
-						BufferSizesName = ralloc_asprintf(ParseState, "%s.BufferSizes", IABVariable->name);
-						bSideTable = true;
-					}
-					else
-					{
-						// Function argument path
-						Index = Buffers.GetIndex(Texture);
-						FormatIndex = Index;
-					}
-					check(Index >= 0 && Index <= 30);
-					
 					ralloc_asprintf_append(buffer, "(");
 					
 					bool bIsStructuredBuffer = (Texture->type->inner_type->is_record() || !strncmp(Texture->type->name, "RWStructuredBuffer<", 19) || !strncmp(Texture->type->name, "StructuredBuffer<", 17));
@@ -2802,6 +2731,13 @@ protected:
 				{
 					deref->image->accept(this);
 					ralloc_asprintf_append(buffer, ".write(");
+					
+					bool bSwizzleSample = Backend.bSwizzleSample;
+					if (bSwizzleSample)
+					{
+						ralloc_asprintf_append(buffer, "swizzle_sample(");
+					}
+					
 					// @todo Zebra: Below is a terrible hack - the input to write is always vec<T, 4>,
 					// 				but the type T comes from the texture type.  
 					if(src_elements == 1)
@@ -2857,6 +2793,11 @@ protected:
 								break;
 						}
 					}
+					
+					if (bSwizzleSample)
+					{
+						ralloc_asprintf_append(buffer, ", %s[%d])", BufferSizesName, Index);
+					}
 
 					//#todo-rco: Add language spec to know if indices need to be uint
 					ralloc_asprintf_append(buffer, ",(uint");
@@ -2897,35 +2838,6 @@ protected:
 				}
 				else
 				{
-					auto* Texture = deref->image->variable_referenced();
-					int Index = 0;
-					char const* BufferSizesName = "BufferSizes";
-					int FormatIndex = -1;
-					bool bSideTable = bInsertSideTable;
-					if (Texture->mode == ir_var_temporary)
-					{
-						// IAB sampling path
-						ir_variable* IABVariable = Backend.IABVariablesMap.FindChecked(Texture);
-						int FieldIndex = IABVariable->type->field_index(Texture->name);
-						for (int i = 0; i < FieldIndex; i++)
-						{
-							if (IABVariable->type->fields.structure[i].type->sampler_buffer)
-							{
-								Index++;
-							}
-						}
-						
-						BufferSizesName = ralloc_asprintf(ParseState, "%s.BufferSizes", IABVariable->name);
-						bSideTable = true;
-					}
-					else
-					{
-						// Function argument path
-						Index = Buffers.GetIndex(Texture);
-						FormatIndex = Index;
-					}
-					check(Index >= 0 && Index <= 30);
-					
 					bool bIsStructuredBuffer = (Texture->type->inner_type->is_record() || !strncmp(Texture->type->name, "RWStructuredBuffer<", 19) || !strncmp(Texture->type->name, "StructuredBuffer<", 17));
 					bool bIsByteAddressBuffer = (!strncmp(Texture->type->name, "RWByteAddressBuffer", 19) || !strncmp(Texture->type->name, "ByteAddressBuffer", 17));
                     bool bIsInvariantType = Texture->invariant;
@@ -6395,7 +6307,7 @@ void FMetalCodeBackend::CallPatchConstantFunction(_mesa_glsl_parse_state* ParseS
 	pv_if->then_instructions.push_tail(thread_if);
 }
 
-FMetalCodeBackend::FMetalCodeBackend(FMetalTessellationOutputs& TessOutputAttribs, unsigned int InHlslCompileFlags, EHlslCompileTarget InTarget, uint8 InVersion, EMetalGPUSemantics bInDesktop, EMetalTypeBufferMode InTypedMode, uint32 InMaxUnrollLoops, bool bInZeroInitialise, bool bInBoundsChecks, bool bInAllFastIntriniscs, bool bInForceInvariance) :
+FMetalCodeBackend::FMetalCodeBackend(FMetalTessellationOutputs& TessOutputAttribs, unsigned int InHlslCompileFlags, EHlslCompileTarget InTarget, uint8 InVersion, EMetalGPUSemantics bInDesktop, EMetalTypeBufferMode InTypedMode, uint32 InMaxUnrollLoops, bool bInZeroInitialise, bool bInBoundsChecks, bool bInAllFastIntriniscs, bool bInForceInvariance, bool bInSwizzleSample) :
 	FCodeBackend(InHlslCompileFlags, HCT_FeatureLevelES3_1),
 	TessAttribs(TessOutputAttribs),
 	InvariantBuffers(0),
@@ -6410,6 +6322,7 @@ FMetalCodeBackend::FMetalCodeBackend(FMetalTessellationOutputs& TessOutputAttrib
 	MaxUnrollLoops = InMaxUnrollLoops;
 	bZeroInitialise = bInZeroInitialise;
 	bBoundsChecks = bInBoundsChecks;
+	bSwizzleSample = bInSwizzleSample;
 	bAllowFastIntriniscs = bInAllFastIntriniscs;
 	bForceInvariance = bInForceInvariance;
 	
diff --git a/Engine/Source/Developer/Apple/MetalShaderFormat/Private/MetalBackend.h b/Engine/Source/Developer/Apple/MetalShaderFormat/Private/MetalBackend.h
--- a/Engine/Source/Developer/Apple/MetalShaderFormat/Private/MetalBackend.h
+++ b/Engine/Source/Developer/Apple/MetalShaderFormat/Private/MetalBackend.h
@@ -86,7 +86,7 @@ static const int32 MaxMetalSamplers = 16;
 // Generates Metal compliant code from IR tokens
 struct FMetalCodeBackend : public FCodeBackend
 {
-	FMetalCodeBackend(FMetalTessellationOutputs& Attribs, unsigned int InHlslCompileFlags, EHlslCompileTarget InTarget, uint8 Version, EMetalGPUSemantics bInDesktop, EMetalTypeBufferMode InTypedMode, uint32 MaxUnrollLoops, bool bInZeroInitialise, bool bInBoundsChecks, bool bInAllFastIntriniscs, bool bForceInvariance);
+	FMetalCodeBackend(FMetalTessellationOutputs& Attribs, unsigned int InHlslCompileFlags, EHlslCompileTarget InTarget, uint8 Version, EMetalGPUSemantics bInDesktop, EMetalTypeBufferMode InTypedMode, uint32 MaxUnrollLoops, bool bInZeroInitialise, bool bInBoundsChecks, bool bInAllFastIntriniscs, bool bForceInvariance, bool bSwizzleSample);
 
 	virtual char* GenerateCode(struct exec_list* ir, struct _mesa_glsl_parse_state* ParseState, EHlslShaderFrequency Frequency) override;
 
@@ -130,6 +130,7 @@ struct FMetalCodeBackend : public FCodeBackend
 	bool bAllowFastIntriniscs;
 	bool bExplicitDepthWrites;
 	bool bForceInvariance;
+	bool bSwizzleSample;
 
 	bool bIsTessellationVSHS = false;
 	unsigned int inputcontrolpoints = 0;
@@ -139,3 +140,109 @@ struct FMetalCodeBackend : public FCodeBackend
 
 struct FShaderCompilerEnvironment;
 bool IsRemoteBuildingConfigured(const FShaderCompilerEnvironment* InEnvironment = nullptr);
+
+struct SDMARange
+{
+	unsigned SourceCB;
+	unsigned SourceOffset;
+	unsigned Size;
+	unsigned DestCBIndex;
+	unsigned DestCBPrecision;
+	unsigned DestOffset;
+	
+	bool operator <(SDMARange const & Other) const
+	{
+		if (SourceCB == Other.SourceCB)
+		{
+			return SourceOffset < Other.SourceOffset;
+		}
+		
+		return SourceCB < Other.SourceCB;
+	}
+};
+typedef std::list<SDMARange> TDMARangeList;
+typedef std::map<unsigned, TDMARangeList> TCBDMARangeMap;
+
+
+static void InsertRange( TCBDMARangeMap& CBAllRanges, unsigned SourceCB, unsigned SourceOffset, unsigned Size, unsigned DestCBIndex, unsigned DestCBPrecision, unsigned DestOffset )
+{
+	check(SourceCB < (1 << 12));
+	check(DestCBIndex < (1 << 12));
+	check(DestCBPrecision < (1 << 8));
+	unsigned SourceDestCBKey = (SourceCB << 20) | (DestCBIndex << 8) | DestCBPrecision;
+	SDMARange Range = { SourceCB, SourceOffset, Size, DestCBIndex, DestCBPrecision, DestOffset };
+	
+	TDMARangeList& CBRanges = CBAllRanges[SourceDestCBKey];
+	//printf("* InsertRange: %08x\t%u:%u - %u:%c:%u:%u\n", SourceDestCBKey, SourceCB, SourceOffset, DestCBIndex, DestCBPrecision, DestOffset, Size);
+	if (CBRanges.empty())
+	{
+		CBRanges.push_back(Range);
+	}
+	else
+	{
+		TDMARangeList::iterator Prev = CBRanges.end();
+		bool bAdded = false;
+		for (auto Iter = CBRanges.begin(); Iter != CBRanges.end(); ++Iter)
+		{
+			if (SourceOffset + Size <= Iter->SourceOffset)
+			{
+				if (Prev == CBRanges.end())
+				{
+					CBRanges.push_front(Range);
+				}
+				else
+				{
+					CBRanges.insert(Iter, Range);
+				}
+				
+				bAdded = true;
+				break;
+			}
+			
+			Prev = Iter;
+		}
+		
+		if (!bAdded)
+		{
+			CBRanges.push_back(Range);
+		}
+		
+		if (CBRanges.size() > 1)
+		{
+			// Try to merge ranges
+			bool bDirty = false;
+			do
+			{
+				bDirty = false;
+				TDMARangeList NewCBRanges;
+				for (auto Iter = CBRanges.begin(); Iter != CBRanges.end(); ++Iter)
+				{
+					if (Iter == CBRanges.begin())
+					{
+						Prev = CBRanges.begin();
+					}
+					else
+					{
+						if (Prev->SourceOffset + Prev->Size == Iter->SourceOffset && Prev->DestOffset + Prev->Size == Iter->DestOffset)
+						{
+							SDMARange Merged = *Prev;
+							Merged.Size = Prev->Size + Iter->Size;
+							NewCBRanges.pop_back();
+							NewCBRanges.push_back(Merged);
+							++Iter;
+							NewCBRanges.insert(NewCBRanges.end(), Iter, CBRanges.end());
+							bDirty = true;
+							break;
+						}
+					}
+					
+					NewCBRanges.push_back(*Iter);
+					Prev = Iter;
+				}
+				
+				CBRanges.swap(NewCBRanges);
+			}
+			while (bDirty);
+		}
+	}
+}
diff --git a/Engine/Source/Developer/Apple/MetalShaderFormat/Private/MetalDerivedData.cpp b/Engine/Source/Developer/Apple/MetalShaderFormat/Private/MetalDerivedData.cpp
--- a/Engine/Source/Developer/Apple/MetalShaderFormat/Private/MetalDerivedData.cpp
+++ b/Engine/Source/Developer/Apple/MetalShaderFormat/Private/MetalDerivedData.cpp
@@ -375,6 +375,13 @@ bool FMetalShaderOutputCooker::Build(TArray<uint8>& OutData)
 	bool const bZeroInitialise = Input.Environment.CompilerFlags.Contains(CFLAG_ZeroInitialise);
 	bool const bBoundsChecks = Input.Environment.CompilerFlags.Contains(CFLAG_BoundsChecking);
 
+	bool bSwizzleSample = false;
+	FString const* Swizzle = Input.Environment.GetDefinitions().Find(TEXT("METAL_SWIZZLE_SAMPLES"));
+	if (Swizzle)
+	{
+		LexFromString(bSwizzleSample, *(*Swizzle));
+	}
+
 	bool bAllowFastIntriniscs = false;
 	FString const* FastIntrinsics = Input.Environment.GetDefinitions().Find(TEXT("METAL_USE_FAST_INTRINSICS"));
 	if (FastIntrinsics)
@@ -391,7 +398,7 @@ bool FMetalShaderOutputCooker::Build(TArray<uint8>& OutData)
 	}
 
 	FMetalTessellationOutputs Attribs;
-	FMetalCodeBackend MetalBackEnd(Attribs, CCFlags, MetalCompilerTarget, VersionEnum, Semantics, TypeMode, MaxUnrollLoops, bZeroInitialise, bBoundsChecks, bAllowFastIntriniscs, bForceInvariance);
+	FMetalCodeBackend MetalBackEnd(Attribs, CCFlags, MetalCompilerTarget, VersionEnum, Semantics, TypeMode, MaxUnrollLoops, bZeroInitialise, bBoundsChecks, bAllowFastIntriniscs, bForceInvariance, bSwizzleSample);
 	FMetalLanguageSpec MetalLanguageSpec(VersionEnum);
 
 	int32 Result = 0;
diff --git a/Engine/Source/Runtime/Apple/MetalRHI/Private/MetalCommandEncoder.cpp b/Engine/Source/Runtime/Apple/MetalRHI/Private/MetalCommandEncoder.cpp
--- a/Engine/Source/Runtime/Apple/MetalRHI/Private/MetalCommandEncoder.cpp
+++ b/Engine/Source/Runtime/Apple/MetalRHI/Private/MetalCommandEncoder.cpp
@@ -1470,6 +1470,18 @@ void FMetalCommandEncoder::SetShaderTexture(mtlpp::FunctionType FunctionType, FM
 	
 	if (Texture)
 	{
+		uint8 Swizzle[4] = {0,0,0,0};
+		assert(sizeof(Swizzle) == sizeof(uint32));
+		if (Texture.GetPixelFormat() == mtlpp::PixelFormat::X32_Stencil8
+#if PLATFORM_MAC
+		 ||	Texture.GetPixelFormat() == mtlpp::PixelFormat::X24_Stencil8
+#endif
+		)
+		{
+			Swizzle[0] = Swizzle[1] = Swizzle[2] = Swizzle[3] = 1;
+		}
+		FMemory::Memcpy(&ShaderBuffers[uint32(FunctionType)].Lengths[(ML_MaxBuffers*2)+(index*2)], Swizzle, sizeof(Swizzle));
+		ShaderBuffers[uint32(FunctionType)].Lengths[(ML_MaxBuffers*2)+(index*2)+1] = 0;
 		TextureBindingHistory.Add(ns::AutoReleased<FMetalTexture>(Texture));
 	}
 }
diff --git a/Engine/Source/Runtime/Apple/MetalRHI/Private/MetalCommandEncoder.h b/Engine/Source/Runtime/Apple/MetalRHI/Private/MetalCommandEncoder.h
--- a/Engine/Source/Runtime/Apple/MetalRHI/Private/MetalCommandEncoder.h
+++ b/Engine/Source/Runtime/Apple/MetalRHI/Private/MetalCommandEncoder.h
@@ -441,7 +441,7 @@ private:
 		/** The usage mask for the bound resource or 0 */
 		mtlpp::ResourceUsage Usage[ML_MaxBuffers];
 		/** The bound buffer lengths */
-		uint32 Lengths[ML_MaxBuffers*2];
+		uint32 Lengths[(ML_MaxBuffers*2) + (ML_MaxTextures*2)];
         /** A bitmask for which buffers were bound by the application where a bit value of 1 is bound and 0 is unbound. */
         uint32 Bound;
 	};
diff --git a/Engine/Source/Runtime/Apple/MetalRHI/Private/MetalStateCache.cpp b/Engine/Source/Runtime/Apple/MetalRHI/Private/MetalStateCache.cpp
--- a/Engine/Source/Runtime/Apple/MetalRHI/Private/MetalStateCache.cpp
+++ b/Engine/Source/Runtime/Apple/MetalRHI/Private/MetalStateCache.cpp
@@ -1473,15 +1473,14 @@ void FMetalStateCache::SetShaderResourceView(FMetalContext* Context, EShaderFreq
 			ns::AutoReleased<FMetalTexture> Tex;
 			Tex = SRV->GetLinearTexture(false);
 			
-			uint32 PackedLen = (Tex.GetWidth() | (Tex.GetHeight() << 16));
 			SetShaderTexture(ShaderStage, Tex, BindIndex, mtlpp::ResourceUsage(mtlpp::ResourceUsage::Read|mtlpp::ResourceUsage::Sample));
 			if (VB)
             {
-                SetShaderBuffer(ShaderStage, VB->Buffer, VB->Data, 0, PackedLen, BindIndex, mtlpp::ResourceUsage::Read, (EPixelFormat)SRV->Format);
+                SetShaderBuffer(ShaderStage, VB->Buffer, VB->Data, 0, VB->GetSize(), BindIndex, mtlpp::ResourceUsage::Read, (EPixelFormat)SRV->Format);
             }
             else if (IB)
             {
-                SetShaderBuffer(ShaderStage, IB->Buffer, nil, 0, PackedLen, BindIndex, mtlpp::ResourceUsage::Read, (EPixelFormat)SRV->Format);
+                SetShaderBuffer(ShaderStage, IB->Buffer, nil, 0, IB->GetSize(), BindIndex, mtlpp::ResourceUsage::Read, (EPixelFormat)SRV->Format);
             }
 		}
 		else if (VB)
@@ -1558,8 +1557,7 @@ void FMetalStateCache::SetShaderUnorderedAccessView(EShaderFrequency ShaderStage
 				Tex = UAV->SourceView->GetLinearTexture(true);
 				SetShaderTexture(ShaderStage, Tex, BindIndex, mtlpp::ResourceUsage(mtlpp::ResourceUsage::Read | mtlpp::ResourceUsage::Write));
                 
-                uint32 PackedLen = (Tex.GetWidth() | (Tex.GetHeight() << 16));
-                SetShaderBuffer(ShaderStage, VertexBuffer->Buffer, VertexBuffer->Data, 0, PackedLen, BindIndex, mtlpp::ResourceUsage(mtlpp::ResourceUsage::Read | mtlpp::ResourceUsage::Write), (EPixelFormat)UAV->SourceView->Format);
+                SetShaderBuffer(ShaderStage, VertexBuffer->Buffer, VertexBuffer->Data, 0, VertexBuffer->GetSize(), BindIndex, mtlpp::ResourceUsage(mtlpp::ResourceUsage::Read | mtlpp::ResourceUsage::Write), (EPixelFormat)UAV->SourceView->Format);
 			}
 			else
 			{
@@ -1575,8 +1573,7 @@ void FMetalStateCache::SetShaderUnorderedAccessView(EShaderFrequency ShaderStage
 				Tex = UAV->SourceView->GetLinearTexture(true);
 				SetShaderTexture(ShaderStage, Tex, BindIndex, mtlpp::ResourceUsage(mtlpp::ResourceUsage::Read | mtlpp::ResourceUsage::Write));
 				
-				uint32 PackedLen = (Tex.GetWidth() | (Tex.GetHeight() << 16));
-				SetShaderBuffer(ShaderStage, IndexBuffer->Buffer, nullptr, 0, PackedLen, BindIndex, mtlpp::ResourceUsage(mtlpp::ResourceUsage::Read | mtlpp::ResourceUsage::Write), (EPixelFormat)UAV->SourceView->Format);
+				SetShaderBuffer(ShaderStage, IndexBuffer->Buffer, nullptr, 0, IndexBuffer->GetSize(), BindIndex, mtlpp::ResourceUsage(mtlpp::ResourceUsage::Read | mtlpp::ResourceUsage::Write), (EPixelFormat)UAV->SourceView->Format);
 			}
 			else
 			{
