Disabled ue4::safe_array for non-mobile GPU targets in Metal cross compiler.

diff --git a/Engine/Source/Developer/Apple/MetalShaderFormat/Private/MetalBackend.cpp b/Engine/Source/Developer/Apple/MetalShaderFormat/Private/MetalBackend.cpp
--- a/Engine/Source/Developer/Apple/MetalShaderFormat/Private/MetalBackend.cpp
+++ b/Engine/Source/Developer/Apple/MetalShaderFormat/Private/MetalBackend.cpp
@@ -1259,7 +1259,7 @@ protected:
 						ralloc_asprintf_append(buffer, "/*ir_var_out, is_patch_constant*/");
 					}
 				}
-				else if ((var->mode == ir_var_auto || var->mode == ir_var_temporary) && var->type->is_array())
+				else if ((var->mode == ir_var_auto || var->mode == ir_var_temporary) && var->type->is_array() && Backend.bEnableUE4SafeArray)
 				{
 					ralloc_asprintf_append(buffer, "ue4::safe_array<");
 					print_type_pre(var->type->element_type());
@@ -6323,7 +6323,7 @@ void FMetalCodeBackend::CallPatchConstantFunction(_mesa_glsl_parse_state* ParseS
 	pv_if->then_instructions.push_tail(thread_if);
 }
 
-FMetalCodeBackend::FMetalCodeBackend(FMetalTessellationOutputs& TessOutputAttribs, unsigned int InHlslCompileFlags, EHlslCompileTarget InTarget, uint8 InVersion, EMetalGPUSemantics bInDesktop, EMetalTypeBufferMode InTypedMode, uint32 InMaxUnrollLoops, bool bInZeroInitialise, bool bInBoundsChecks, bool bInAllFastIntriniscs, bool bInForceInvariance, bool bInSwizzleSample) :
+FMetalCodeBackend::FMetalCodeBackend(FMetalTessellationOutputs& TessOutputAttribs, unsigned int InHlslCompileFlags, EHlslCompileTarget InTarget, uint8 InVersion, EMetalGPUSemantics bInDesktop, EMetalTypeBufferMode InTypedMode, uint32 InMaxUnrollLoops, bool bInZeroInitialise, bool bInBoundsChecks, bool bInAllFastIntriniscs, bool bInForceInvariance, bool bInSwizzleSample, bool bInEnableUE4SafeArray) :
 	FCodeBackend(InHlslCompileFlags, HCT_FeatureLevelES3_1),
 	TessAttribs(TessOutputAttribs),
 	InvariantBuffers(0),
@@ -6339,6 +6339,7 @@ FMetalCodeBackend::FMetalCodeBackend(FMetalTessellationOutputs& TessOutputAttrib
 	bZeroInitialise = bInZeroInitialise;
 	bBoundsChecks = bInBoundsChecks;
 	bSwizzleSample = bInSwizzleSample;
+	bEnableUE4SafeArray = bInEnableUE4SafeArray;
 	bAllowFastIntriniscs = bInAllFastIntriniscs;
 	bForceInvariance = bInForceInvariance;
 	
diff --git a/Engine/Source/Developer/Apple/MetalShaderFormat/Private/MetalBackend.h b/Engine/Source/Developer/Apple/MetalShaderFormat/Private/MetalBackend.h
--- a/Engine/Source/Developer/Apple/MetalShaderFormat/Private/MetalBackend.h
+++ b/Engine/Source/Developer/Apple/MetalShaderFormat/Private/MetalBackend.h
@@ -88,7 +88,7 @@ static const int32 MaxMetalSamplers = 16;
 // Generates Metal compliant code from IR tokens
 struct FMetalCodeBackend : public FCodeBackend
 {
-	FMetalCodeBackend(FMetalTessellationOutputs& Attribs, unsigned int InHlslCompileFlags, EHlslCompileTarget InTarget, uint8 Version, EMetalGPUSemantics bInDesktop, EMetalTypeBufferMode InTypedMode, uint32 MaxUnrollLoops, bool bInZeroInitialise, bool bInBoundsChecks, bool bInAllFastIntriniscs, bool bForceInvariance, bool bSwizzleSample);
+	FMetalCodeBackend(FMetalTessellationOutputs& Attribs, unsigned int InHlslCompileFlags, EHlslCompileTarget InTarget, uint8 Version, EMetalGPUSemantics bInDesktop, EMetalTypeBufferMode InTypedMode, uint32 MaxUnrollLoops, bool bInZeroInitialise, bool bInBoundsChecks, bool bInAllFastIntriniscs, bool bInForceInvariance, bool bInSwizzleSample, bool bInEnableUE4SafeArray);
 
 	virtual char* GenerateCode(struct exec_list* ir, struct _mesa_glsl_parse_state* ParseState, EHlslShaderFrequency Frequency) override;
 
@@ -133,6 +133,7 @@ struct FMetalCodeBackend : public FCodeBackend
 	bool bExplicitDepthWrites;
 	bool bForceInvariance;
 	bool bSwizzleSample;
+	bool bEnableUE4SafeArray; // enables ue4::safe_array (enabled for iOS, disabled for macOS due to AMD driver bug)
 
 	bool bIsTessellationVSHS = false;
 	unsigned int inputcontrolpoints = 0;
diff --git a/Engine/Source/Developer/Apple/MetalShaderFormat/Private/MetalDerivedData.cpp b/Engine/Source/Developer/Apple/MetalShaderFormat/Private/MetalDerivedData.cpp
--- a/Engine/Source/Developer/Apple/MetalShaderFormat/Private/MetalDerivedData.cpp
+++ b/Engine/Source/Developer/Apple/MetalShaderFormat/Private/MetalDerivedData.cpp
@@ -420,6 +420,9 @@ bool FMetalShaderOutputCooker::Build(TArray<uint8>& OutData)
 		bForceInvariance = true;
 	}
 	
+	// Only enable ue4::safe_array for iOS, due to AMD driver bugs on macOS
+	const bool bEnableUE4SafeArray = (Semantics == EMetalGPUSemanticsMobile);
+	
 	FMetalShaderOutputMetaData OutputData;
 	FMetalTessellationOutputs Attribs;
 	
@@ -440,7 +443,7 @@ bool FMetalShaderOutputCooker::Build(TArray<uint8>& OutData)
 #endif
 #endif
 	{
-		FMetalCodeBackend MetalBackEnd(Attribs, CCFlags, MetalCompilerTarget, VersionEnum, Semantics, TypeMode, MaxUnrollLoops, bZeroInitialise, bBoundsChecks, bAllowFastIntriniscs, bForceInvariance, bSwizzleSample);
+		FMetalCodeBackend MetalBackEnd(Attribs, CCFlags, MetalCompilerTarget, VersionEnum, Semantics, TypeMode, MaxUnrollLoops, bZeroInitialise, bBoundsChecks, bAllowFastIntriniscs, bForceInvariance, bSwizzleSample, bEnableUE4SafeArray);
 		FMetalLanguageSpec MetalLanguageSpec(VersionEnum);
 
 		FHlslCrossCompilerContext CrossCompilerContext(CCFlags, Frequency, HlslCompilerTarget);
