UE-72092 - Action RPG failed to package for iOS with remote shader compiling

diff --git a/Engine/Source/Developer/Apple/MetalShaderFormat/Private/MetalShaderCompiler.cpp b/Engine/Source/Developer/Apple/MetalShaderFormat/Private/MetalShaderCompiler.cpp
--- a/Engine/Source/Developer/Apple/MetalShaderFormat/Private/MetalShaderCompiler.cpp
+++ b/Engine/Source/Developer/Apple/MetalShaderFormat/Private/MetalShaderCompiler.cpp
@@ -370,7 +370,7 @@ static uint32 GetMaxArgLength()
     return ARG_MAX;
 #else
     // Ask the remote machine via "getconf ARG_MAX"
-    return 1024 * 256;
+	return 1024;
 #endif
 }
 
@@ -2494,7 +2494,7 @@ bool StripShader_Metal(TArray<uint8>& Code, class FString const& DebugPath, bool
 						FileHandle->Write((const uint8 *)ShaderSource, ShaderSourceLength);
 						delete FileHandle;
 
-						IFileManager::Get().Move(*DebugFilePath, *TempPath, false, false, true, false);
+						IFileManager::Get().Move(*DebugFilePath, *TempPath, true, false, true, false);
 						IFileManager::Get().Delete(*TempPath);
 					}
 					else
@@ -2737,7 +2737,7 @@ bool FinalizeLibrary_Metal(FName const& Format, FString const& WorkingDir, FStri
 				
 				// Build source file name path
 				UE_LOG(LogShaders, Verbose, TEXT("[%d/%d] %s Main_%0.8x_%0.8x.o"), ++Index, Shaders.Num(), *Format.GetPlainNameString(), Len, CRC);
-				FString SourceFileNameParam = FString::Printf(TEXT("\"%s/Main_%0.8x_%0.8x.o\""), *WorkingDir, Len, CRC);
+				FString SourceFileNameParam = FString::Printf(TEXT("%s/Main_%0.8x_%0.8x.o"), *FPaths::ConvertRelativePathToFull(WorkingDir), Len, CRC);
 				
 				// Remote builds copy file and swizzle Source File Name param
 				if(bBuildingRemotely)
@@ -2825,7 +2825,8 @@ bool FinalizeLibrary_Metal(FName const& Format, FString const& WorkingDir, FStri
 				if (ReturnCode == 0)
 				{
 					// There is problem going to location with spaces using remote copy (at least on Mac no combination of \ and/or "" works) - work around this issue @todo investigate this further
-                    FString LocalCopyLocation = FPaths::Combine(TEXT("/tmp"),FPaths::GetCleanFilename(LibraryPath));
+					FString FileName = FPaths::GetCleanFilename(LibraryPath);
+                    FString LocalCopyLocation = FPaths::Combine(*FPaths::ConvertRelativePathToFull(WorkingDir), FileName);
 						
                     if(bBuildingRemotely && CopyRemoteFileToLocal(RemoteLibPath, LocalCopyLocation))
                     {
