Changed Metal tools detection code to additionally search for the iOS version of Metal compiler in the new location in Xcode 11. This a temporary solution to allow cooking for iOS with Xcode 11 until we convert GetMetalBinaryPath() to use xcrun to find the tools.

diff --git a/Engine/Source/Developer/Apple/MetalShaderFormat/Private/MetalShaderCompiler.cpp b/Engine/Source/Developer/Apple/MetalShaderFormat/Private/MetalShaderCompiler.cpp
--- a/Engine/Source/Developer/Apple/MetalShaderFormat/Private/MetalShaderCompiler.cpp
+++ b/Engine/Source/Developer/Apple/MetalShaderFormat/Private/MetalShaderCompiler.cpp
@@ -487,6 +487,12 @@ FString GetMetalBinaryPath(uint32 ShaderPlatform)
 					MetalToolsPath = FString::Printf(TEXT("%s/Platforms/MacOSX.platform/usr/bin"), *XcodePath);
 				}
 				MetalPath = MetalToolsPath + TEXT("/metal");
+
+				if (!RemoteFileExists(MetalPath) && bIsMobile)
+				{
+					MetalToolsPath = FString::Printf(TEXT("%s/Toolchains/XcodeDefault.xctoolchain/usr/metal/ios/bin"), *XcodePath);
+					MetalPath = MetalToolsPath + TEXT("/metal");
+				}
 			}
 
 			if (RemoteFileExists(MetalPath))
