Fix a texture index allocation in Metal shaders when using hlslcc. Copy CL# 7710435

diff --git a/Engine/Source/Developer/Apple/MetalShaderFormat/Private/MetalUtils.h b/Engine/Source/Developer/Apple/MetalShaderFormat/Private/MetalUtils.h
--- a/Engine/Source/Developer/Apple/MetalShaderFormat/Private/MetalUtils.h
+++ b/Engine/Source/Developer/Apple/MetalShaderFormat/Private/MetalUtils.h
@@ -313,7 +313,7 @@ struct FBuffers
         // This means you will get holes in the slots where we see a typed-resource on the other array.
         // We always leave a hole for BufferSizes at index 31.
         
-        uint32 TypedBuffers = 0;
+        uint32 TypedBuffers = UAVIndices;
         static const uint32 MaxBuffers = 30;
 		check(MaxTextures > 0);
         
