Transform access::write attribute for write-only UAVs.

diff --git a/Engine/Source/Developer/Apple/MetalShaderFormat/Private/MetalBackend.cpp b/Engine/Source/Developer/Apple/MetalShaderFormat/Private/MetalBackend.cpp
--- a/Engine/Source/Developer/Apple/MetalShaderFormat/Private/MetalBackend.cpp
+++ b/Engine/Source/Developer/Apple/MetalShaderFormat/Private/MetalBackend.cpp
@@ -563,7 +563,7 @@ protected:
 	/**
 	 * Print the base type, e.g. vec3.
 	 */
-	void print_base_type(const glsl_type *t)
+	void print_base_type(const glsl_type *t, bool bImageWriteOnly = false)
 	{
 		if (t->base_type == GLSL_TYPE_ARRAY)
 		{
@@ -660,7 +660,14 @@ protected:
 				}
 				else
 				{
-					ralloc_asprintf_append(buffer, ", access::read_write>");
+					if (bImageWriteOnly)
+					{
+						ralloc_asprintf_append(buffer, ", access::write>");
+					}
+					else
+					{
+						ralloc_asprintf_append(buffer, ", access::read_write>");
+					}
 				}
 			}
 		}
@@ -775,9 +782,9 @@ protected:
 	/**
 	 * Print the portion of the type that appears before a variable declaration.
 	 */
-	void print_type_pre(const glsl_type *t)
+	void print_type_pre(const glsl_type *t, bool bImageWriteOnly = false)
 	{
-		print_base_type(t);
+		print_base_type(t, bImageWriteOnly);
 	}
 
 	/**
@@ -992,7 +999,7 @@ protected:
 				}
 				else
 				{
-					print_type_pre(PtrType);
+					print_type_pre(PtrType, (var->image_write && !(var->image_read)));
 					if (var->mode != ir_var_temporary)
 					{
 						ralloc_asprintf_append(buffer, " %s [[ texture(%d) ]]", unique_name(var), BufferIndex);
