Fix Metal shader compile failures caused by missing buffer side table.

diff --git a/Engine/Source/Developer/Apple/MetalShaderFormat/Private/MetalBackend.cpp b/Engine/Source/Developer/Apple/MetalShaderFormat/Private/MetalBackend.cpp
--- a/Engine/Source/Developer/Apple/MetalShaderFormat/Private/MetalBackend.cpp
+++ b/Engine/Source/Developer/Apple/MetalShaderFormat/Private/MetalBackend.cpp
@@ -1351,7 +1351,7 @@ protected:
         bInsertSideTable = Backend.bIsTessellationVSHS;
         if (sig->is_main && Backend.bBoundsChecks)
 		{
-			bInsertSideTable = Backend.bSwizzleSample;
+			bInsertSideTable |= Backend.bSwizzleSample;
             foreach_iter(exec_list_iterator, iter, sig->parameters)
             {
                 ir_variable *const inst = (ir_variable *) iter.get();
@@ -1369,6 +1369,7 @@ protected:
             {
                 ir_variable* BufferSizes = new(ParseState)ir_variable(glsl_type::uint_type, "BufferSizes", ir_var_uniform);
                 BufferSizes->semantic = "u";
+				BufferSizes->used = 1;
                 Buffers.Buffers.Add(BufferSizes);
                 sig->parameters.push_head(BufferSizes);
             }
