*Very* experimental support for using ShaderConductor to translate HLSL to MetalSL (use MaxShaderLanguageVersion=6 - only <= ES31 will compile shaders currently).

Copied from //Tasks/UE4/Dev-UERNDR-354-mtlpp to Dev-Rendering (//UE4/Dev-Rendering)

diff --git a/Engine/Shaders/Private/InstancedStereo.ush b/Engine/Shaders/Private/InstancedStereo.ush
--- a/Engine/Shaders/Private/InstancedStereo.ush
+++ b/Engine/Shaders/Private/InstancedStereo.ush
@@ -27,7 +27,7 @@ static const float4 EyeClipEdge[2] = { float4(-1.0, 0.0, 0.0, 1.0), float4(1.0,
 #endif
 
 // Generated for Metal in GeneratedInstancedStereo.usf
-#if !COMPILER_METAL && (INSTANCED_STEREO || MOBILE_MULTI_VIEW)
+#if !(COMPILER_METAL && MAX_SHADER_LANGUAGE_VERSION < 6) && (INSTANCED_STEREO || MOBILE_MULTI_VIEW)
 ViewState ResolveView(uint ViewIndex)
 {
 	if (ViewIndex == 0)
diff --git a/Engine/Shaders/Private/ShadowFilteringCommon.ush b/Engine/Shaders/Private/ShadowFilteringCommon.ush
--- a/Engine/Shaders/Private/ShadowFilteringCommon.ush
+++ b/Engine/Shaders/Private/ShadowFilteringCommon.ush
@@ -289,17 +289,12 @@ float Manual1x1PCF(float2 ShadowPosition, FPCFSamplerSettings Settings)
 
 float ManualPCF(float2 ShadowPosition, FPCFSamplerSettings Settings)
 {
-
 #if SHADOW_QUALITY == 1
 	return ManualNoFiltering(ShadowPosition, Settings);
-#endif
-
-#if SHADOW_QUALITY == 2
+#elif SHADOW_QUALITY == 2
 	// low quality, 2x2 samples, using and not using gather4
 	return Manual1x1PCF(ShadowPosition, Settings);
-#endif
-
-#if SHADOW_QUALITY == 3
+#elif SHADOW_QUALITY == 3
 	// medium quality, 4x4 samples, using and not using gather4
 	{
 		float2 TexelPos = ShadowPosition * Settings.ShadowBufferSize.xy - 0.5f;	// bias to be consistent with texture filtering hardware
@@ -326,9 +321,7 @@ float ManualPCF(float2 ShadowPosition, FPCFSamplerSettings Settings)
 #endif // FEATURE_GATHER4
 		}
 	}
-#endif
-
-#if FEATURE_GATHER4
+#elif FEATURE_GATHER4
 	// high quality, 6x6 samples, using gather4
 	{
 		float2 TexelPos = ShadowPosition * Settings.ShadowBufferSize.xy - 0.5f;	// bias to be consistent with texture filtering hardware
diff --git a/Engine/Shaders/Public/Platform/Metal/MetalCommon.ush b/Engine/Shaders/Public/Platform/Metal/MetalCommon.ush
--- a/Engine/Shaders/Public/Platform/Metal/MetalCommon.ush
+++ b/Engine/Shaders/Public/Platform/Metal/MetalCommon.ush
@@ -5,7 +5,7 @@
 =============================================================================*/
 
 // Update this GUID to improve shader recompilation for Metal only shaders
-// GUID = 4E19BD71-3C18-4138-B085-451EAC550287
+// GUID = 0FDEF9A8-0978-44ED-B2ED-246924F18E15
   
 #pragma once
 
@@ -28,3 +28,14 @@
 */
 
 #define COMPILER_SUPPORTS_MINMAX3 0
+
+#if IOS
+#if MAX_SHADER_LANGUAGE_VERSION >= 6
+#if PIXELSHADER
+	[[vk::input_attachment_index(0)]] SubpassInput<float4> gl_LastFragData;
+	#define FramebufferFetchES2() gl_LastFragData.SubpassLoad()
+#else
+	#define FramebufferFetchES2() float4(0,0,0,0)
+#endif
+#endif
+#endif
diff --git a/Engine/Source/Developer/Apple/MetalShaderFormat/MetalShaderFormat.Build.cs b/Engine/Source/Developer/Apple/MetalShaderFormat/MetalShaderFormat.Build.cs
--- a/Engine/Source/Developer/Apple/MetalShaderFormat/MetalShaderFormat.Build.cs
+++ b/Engine/Source/Developer/Apple/MetalShaderFormat/MetalShaderFormat.Build.cs
@@ -28,6 +28,14 @@ public class MetalShaderFormat : ModuleRules
 		AddEngineThirdPartyPrivateStaticDependencies(Target, 
 			"HLSLCC"
 			);
+
+		if (Target.Platform == UnrealTargetPlatform.Mac || Target.Platform == UnrealTargetPlatform.Win64)
+		{
+			AddEngineThirdPartyPrivateStaticDependencies(Target,
+				"ShaderConductor",
+				"SPIRVReflect"
+			);
+		}
 			
 		DynamicallyLoadedModuleNames.AddRange(
 			new string[]
diff --git a/Engine/Source/Developer/Apple/MetalShaderFormat/Private/MetalDerivedData.cpp b/Engine/Source/Developer/Apple/MetalShaderFormat/Private/MetalDerivedData.cpp
--- a/Engine/Source/Developer/Apple/MetalShaderFormat/Private/MetalDerivedData.cpp
+++ b/Engine/Source/Developer/Apple/MetalShaderFormat/Private/MetalDerivedData.cpp
@@ -11,6 +11,14 @@
 #include "Misc/Paths.h"
 #include "Misc/Compression.h"
 #include "MetalBackend.h"
+#include "Serialization/MemoryReader.h"
+
+#if PLATFORM_MAC || PLATFORM_WINDOWS
+THIRD_PARTY_INCLUDES_START
+#include "ShaderConductor/ShaderConductor.hpp"
+#include "spirv_reflect.h"
+THIRD_PARTY_INCLUDES_END
+#endif
 
 extern bool ExecRemoteProcess(const TCHAR* Command, const TCHAR* Params, int32* OutReturnCode, FString* OutStdOut, FString* OutStdErr);
 extern FString GetXcodePath();
@@ -364,13 +372,22 @@ bool FMetalShaderOutputCooker::IsBuildThreadsafe() const
 	return false;
 }
 
+struct FMetalShaderOutputMetaData
+{
+	TArray<uint8> TypedBufferFormats;
+	uint32 InvariantBuffers = 0;
+	uint32 TypedBuffers = 0;
+	uint32 TypedUAVs = 0;
+	uint32 ConstantBuffers = 0;
+};
+
 bool FMetalShaderOutputCooker::Build(TArray<uint8>& OutData)
 {
 	Output.bSucceeded = false;
 
-	char* MetalShaderSource = NULL;
-	char* ErrorLog = NULL;
-
+	std::string MetalSource;
+	FString MetalErrors;
+	
 	bool const bZeroInitialise = Input.Environment.CompilerFlags.Contains(CFLAG_ZeroInitialise);
 	bool const bBoundsChecks = Input.Environment.CompilerFlags.Contains(CFLAG_BoundsChecking);
 
@@ -395,52 +412,796 @@ bool FMetalShaderOutputCooker::Build(TArray<uint8>& OutData)
 		// WPO requires that we make all multiply/sincos instructions invariant :(
 		bForceInvariance = true;
 	}
-
+	
+	FMetalShaderOutputMetaData OutputData;
 	FMetalTessellationOutputs Attribs;
-	FMetalCodeBackend MetalBackEnd(Attribs, CCFlags, MetalCompilerTarget, VersionEnum, Semantics, TypeMode, MaxUnrollLoops, bZeroInitialise, bBoundsChecks, bAllowFastIntriniscs, bForceInvariance, bSwizzleSample);
-	FMetalLanguageSpec MetalLanguageSpec(VersionEnum);
-
+	
+	uint32 CRCLen = 0;
+	uint32 CRC = 0;
+	uint32 SourceLen = 0;
 	int32 Result = 0;
-	FHlslCrossCompilerContext CrossCompilerContext(CCFlags, Frequency, HlslCompilerTarget);
-	if (CrossCompilerContext.Init(TCHAR_TO_ANSI(*Input.VirtualSourceFilePath), &MetalLanguageSpec))
+	
+	FString const* UsingTessellationDefine = Input.Environment.GetDefinitions().Find(TEXT("USING_TESSELLATION"));
+	bool bUsingTessellation = ((UsingTessellationDefine != nullptr && FString("1") == *UsingTessellationDefine && Frequency == HSF_VertexShader) || Frequency == HSF_HullShader || Frequency == HSF_DomainShader);
+	
+	// Its going to take a while to get dxc+SPIRV working...
+#if PLATFORM_MAC || PLATFORM_WINDOWS
+	if (VersionEnum < 6)
+#endif
 	{
-		Result = CrossCompilerContext.Run(
-			TCHAR_TO_ANSI(*PreprocessedShader),
-			TCHAR_TO_ANSI(*Input.EntryPointName),
-			&MetalBackEnd,
-			&MetalShaderSource,
-			&ErrorLog
-		) ? 1 : 0;
-	}
+		FMetalCodeBackend MetalBackEnd(Attribs, CCFlags, MetalCompilerTarget, VersionEnum, Semantics, TypeMode, MaxUnrollLoops, bZeroInitialise, bBoundsChecks, bAllowFastIntriniscs, bForceInvariance, bSwizzleSample);
+		FMetalLanguageSpec MetalLanguageSpec(VersionEnum);
 
-	uint32 CRCLen = MetalShaderSource ? (uint32)FCStringAnsi::Strlen(MetalShaderSource) : 0u;
-	uint32 CRC = CRCLen ? FCrc::MemCrc_DEPRECATED(MetalShaderSource, CRCLen) : 0u;
-	uint32 SourceLen = CRCLen;
-	if (MetalShaderSource)
+		FHlslCrossCompilerContext CrossCompilerContext(CCFlags, Frequency, HlslCompilerTarget);
+		if (CrossCompilerContext.Init(TCHAR_TO_ANSI(*Input.VirtualSourceFilePath), &MetalLanguageSpec))
+		{
+			char* MetalShaderSource = NULL;
+			char* ErrorLog = NULL;
+
+			Result = CrossCompilerContext.Run(
+                TCHAR_TO_ANSI(*PreprocessedShader),
+				TCHAR_TO_ANSI(*Input.EntryPointName),
+				&MetalBackEnd,
+				&MetalShaderSource,
+				&ErrorLog
+            ) ? 1 : 0;
+			
+			OutputData.TypedBufferFormats = MetalBackEnd.TypedBufferFormats;
+			OutputData.InvariantBuffers = MetalBackEnd.InvariantBuffers;
+			OutputData.TypedBuffers = MetalBackEnd.TypedBuffers;
+			OutputData.TypedUAVs = MetalBackEnd.TypedUAVs;
+			OutputData.ConstantBuffers = MetalBackEnd.ConstantBuffers;
+			
+			CRCLen = MetalShaderSource ? (uint32)FCStringAnsi::Strlen(MetalShaderSource) : 0u;
+			CRC = CRCLen ? FCrc::MemCrc_DEPRECATED(MetalShaderSource, CRCLen) : 0u;
+			SourceLen = CRCLen;
+			if (MetalShaderSource)
+			{
+				ANSICHAR* Main = FCStringAnsi::Strstr(MetalShaderSource, "Main_00000000_00000000");
+				check(Main);
+				
+				ANSICHAR MainCRC[24];
+				int32 NewLen = FCStringAnsi::Snprintf(MainCRC, 24, "Main_%0.8x_%0.8x", CRCLen, CRC);
+				FMemory::Memcpy(Main, MainCRC, NewLen);
+				
+				uint32 Len = FCStringAnsi::Strlen(TCHAR_TO_ANSI(*Input.GetSourceFilename())) + FCStringAnsi::Strlen(TCHAR_TO_ANSI(*Input.DebugGroupName)) + FCStringAnsi::Strlen(TCHAR_TO_ANSI(*Input.EntryPointName)) + FCStringAnsi::Strlen(MetalShaderSource) + 21;
+				char* Dest = (char*)malloc(Len);
+				FCStringAnsi::Snprintf(Dest, Len, "// ! %s/%s.usf:%s\n%s", (const char*)TCHAR_TO_ANSI(*Input.DebugGroupName), (const char*)TCHAR_TO_ANSI(*Input.GetSourceFilename()), (const char*)TCHAR_TO_ANSI(*Input.EntryPointName), (const char*)MetalShaderSource);
+				free(MetalShaderSource);
+				MetalShaderSource = Dest;
+				SourceLen = (uint32)FCStringAnsi::Strlen(MetalShaderSource);
+			}
+			
+			if (MetalShaderSource)
+			{
+				MetalSource = MetalShaderSource;
+				free(MetalShaderSource);
+			}
+			if (ErrorLog)
+			{
+				MetalErrors = ANSI_TO_TCHAR(ErrorLog);
+				free(ErrorLog);
+			}
+		}
+	}
+#if PLATFORM_MAC || PLATFORM_WINDOWS
+	else
 	{
-		ANSICHAR* Main = FCStringAnsi::Strstr(MetalShaderSource, "Main_00000000_00000000");
-		check(Main);
-
-		ANSICHAR MainCRC[24];
-		int32 NewLen = FCStringAnsi::Snprintf(MainCRC, 24, "Main_%0.8x_%0.8x", CRCLen, CRC);
-		FMemory::Memcpy(Main, MainCRC, NewLen);
-
-		uint32 Len = FCStringAnsi::Strlen(TCHAR_TO_ANSI(*Input.GetSourceFilename())) + FCStringAnsi::Strlen(TCHAR_TO_ANSI(*Input.DebugGroupName)) + FCStringAnsi::Strlen(TCHAR_TO_ANSI(*Input.EntryPointName)) + FCStringAnsi::Strlen(MetalShaderSource) + 21;
-		char* Dest = (char*)malloc(Len);
-		FCStringAnsi::Snprintf(Dest, Len, "// ! %s/%s.usf:%s\n%s", (const char*)TCHAR_TO_ANSI(*Input.DebugGroupName), (const char*)TCHAR_TO_ANSI(*Input.GetSourceFilename()), (const char*)TCHAR_TO_ANSI(*Input.EntryPointName), (const char*)MetalShaderSource);
-		free(MetalShaderSource);
-		MetalShaderSource = Dest;
-		SourceLen = (uint32)FCStringAnsi::Strlen(MetalShaderSource);
+		ShaderConductor::Compiler::Options Options;
+		Options.removeUnusedGlobals = true;
+		// Metal is inverted vs. SPIRV for some curious reason
+		Options.packMatricesInRowMajor = false;
+		Options.enableDebugInfo = false;
+		Options.enable16bitTypes = false;
+		Options.disableOptimizations = false;
+		
+        ShaderConductor::Compiler::SourceDesc SourceDesc;
+		
+		std::string SourceData(TCHAR_TO_UTF8(*PreprocessedShader));
+		std::string FileName(TCHAR_TO_UTF8(*Input.VirtualSourceFilePath));
+		std::string EntryPointName(TCHAR_TO_UTF8(*Input.EntryPointName));
+		
+        SourceDesc.source = SourceData.c_str();
+        SourceDesc.fileName = FileName.c_str();
+        SourceDesc.entryPoint = EntryPointName.c_str();
+		SourceDesc.numDefines = 0;
+		SourceDesc.defines = nullptr;
+        
+        switch (Frequency)
+        {
+            case HSF_VertexShader:
+            {
+                SourceDesc.stage = ShaderConductor::ShaderStage::VertexShader;
+                break;
+            }
+            case HSF_PixelShader:
+            {
+                SourceDesc.stage = ShaderConductor::ShaderStage::PixelShader;
+                break;
+            }
+            case HSF_GeometryShader:
+            {
+                SourceDesc.stage = ShaderConductor::ShaderStage::GeometryShader;
+                break;
+            }
+            case HSF_HullShader:
+            {
+                SourceDesc.stage = ShaderConductor::ShaderStage::HullShader;
+                break;
+            }
+            case HSF_DomainShader:
+            {
+                SourceDesc.stage = ShaderConductor::ShaderStage::DomainShader;
+                break;
+            }
+            case HSF_ComputeShader:
+            {
+                SourceDesc.stage = ShaderConductor::ShaderStage::ComputeShader;
+                break;
+            }
+            default:
+            {
+                break;
+            }
+        }
+		
+		ShaderConductor::Blob* RewriteBlob = nullptr;
+		
+        // Can't rewrite tessellation shaders because that'll prevent combining Vertex and Hull
+        if (!bUsingTessellation)
+        {
+            ShaderConductor::Compiler::ResultDesc Results = ShaderConductor::Compiler::Rewrite(SourceDesc, Options);
+			RewriteBlob = Results.target;
+			
+			SourceData.clear();
+            SourceData.resize(RewriteBlob->Size());
+			FCStringAnsi::Strncpy(const_cast<char*>(SourceData.c_str()), (const char*)RewriteBlob->Data(), RewriteBlob->Size());
+			SourceDesc.source = SourceData.c_str();
+			
+            if (bDumpDebugInfo)
+            {
+                FArchive* FileWriter = IFileManager::Get().CreateFileWriter(*(Input.DumpDebugInfoPath / FPaths::GetBaseFilename(Input.GetSourceFilename()) + TEXT(".dxc.hlsl")));
+                if (FileWriter)
+                {
+                    FileWriter->Serialize(const_cast<void*>(RewriteBlob->Data()), RewriteBlob->Size());
+                    FileWriter->Close();
+                    delete FileWriter;
+                }
+            }
+        }
+		
+		Options.removeUnusedGlobals = false;
+
+		ShaderConductor::Compiler::TargetDesc TargetDesc;
+		
+		FString MetaData = FString::Printf(TEXT("// ! %s/%s.usf:%s\n"), *Input.DebugGroupName, *Input.GetSourceFilename(), *Input.EntryPointName);
+		FString EntryPoint = Input.EntryPointName;
+		EHlslShaderFrequency Freq = Frequency;
+
+		TargetDesc.language = ShaderConductor::ShadingLanguage::SpirV;
+		ShaderConductor::Compiler::ResultDesc Results = ShaderConductor::Compiler::Compile(SourceDesc, Options, TargetDesc);
+		Result = (Results.hasError) ? 0 : 1;
+		if (!Results.hasError)
+		{
+			// Now perform reflection on the SPIRV and tweak any decorations that we need to.
+			// This used to be done via JSON, but that was slow and alloc happy so use SPIRV-Reflect instead.
+			spv_reflect::ShaderModule Reflection(Results.target->Size(), Results.target->Data());
+			check(Reflection.GetResult() == SPV_REFLECT_RESULT_SUCCESS);
+			
+			SpvReflectResult SPVRResult = SPV_REFLECT_RESULT_NOT_READY;
+			uint32 Count = 0;
+			TArray<SpvReflectDescriptorBinding*> Bindings;
+			TArray<SpvReflectInterfaceVariable*> InputVars;
+			TArray<SpvReflectInterfaceVariable*> OutputVars;
+			
+			uint8 UAVIndices = 0xff;
+			uint32 BufferIndices = 0xffffffff;
+			uint64 TextureIndices = 0xffffffffffffffff;
+			uint16 SamplerIndices = 0xffff;
+			
+			FString UAVString;
+			FString SRVString;
+			FString SMPString;
+			FString UBOString;
+			FString GLOString;
+			FString PAKString;
+			FString INPString;
+			FString OUTString;
+			FString WKGString;
+			
+			if (Frequency == HSF_ComputeShader)
+			{
+				size_t ThreadSizeIdx = SourceData.find("numthreads");
+				check(ThreadSizeIdx != std::string::npos);
+				char const* String = SourceData.c_str() + ThreadSizeIdx;
+				uint32 WorkgroupSize[3];
+				size_t Found = sscanf(String, "numthreads(%d, %d, %d)", &WorkgroupSize[0], &WorkgroupSize[1], &WorkgroupSize[2]);
+				check(Found == 3);
+				for (uint32 i = 0; i < Found; i++)
+				{
+					WKGString += FString::Printf(TEXT("%s%u"), WKGString.Len() ? TEXT(", ") : TEXT(""), WorkgroupSize[i]);
+				}
+			}
+			
+			Count = 0;
+			SPVRResult = Reflection.EnumerateDescriptorBindings(&Count, nullptr);
+			check(SPVRResult == SPV_REFLECT_RESULT_SUCCESS);
+			Bindings.SetNum(Count);
+			SPVRResult = Reflection.EnumerateDescriptorBindings(&Count, Bindings.GetData());
+			check(SPVRResult == SPV_REFLECT_RESULT_SUCCESS);
+			if (Count > 0)
+			{
+				TArray<SpvReflectDescriptorBinding*> UniformBindings;
+				TArray<SpvReflectDescriptorBinding*> SamplerBindings;
+				TArray<SpvReflectDescriptorBinding*> TextureSRVBindings;
+				TArray<SpvReflectDescriptorBinding*> TextureUAVBindings;
+				TArray<SpvReflectDescriptorBinding*> TBufferSRVBindings;
+				TArray<SpvReflectDescriptorBinding*> TBufferUAVBindings;
+				TArray<SpvReflectDescriptorBinding*> SBufferSRVBindings;
+				TArray<SpvReflectDescriptorBinding*> SBufferUAVBindings;
+				
+				// Extract all the bindings first so that we process them in order - this lets us assign UAVs before other resources
+				// Which is necessary to match the D3D binding scheme.
+				for (auto const& Binding : Bindings)
+				{
+					if (Binding->accessed)
+					{
+						switch(Binding->resource_type)
+						{
+							case SPV_REFLECT_RESOURCE_FLAG_CBV:
+							{
+								check(Binding->descriptor_type == SPV_REFLECT_DESCRIPTOR_TYPE_UNIFORM_BUFFER);
+								UniformBindings.Add(Binding);
+								break;
+							}
+							case SPV_REFLECT_RESOURCE_FLAG_SAMPLER:
+							{
+								check(Binding->descriptor_type == SPV_REFLECT_DESCRIPTOR_TYPE_SAMPLER);
+								SamplerBindings.Add(Binding);
+								break;
+							}
+							case SPV_REFLECT_RESOURCE_FLAG_SRV:
+							{
+								switch(Binding->descriptor_type)
+								{
+									case SPV_REFLECT_DESCRIPTOR_TYPE_SAMPLED_IMAGE:
+									{
+										TextureSRVBindings.Add(Binding);
+										break;
+									}
+									case SPV_REFLECT_DESCRIPTOR_TYPE_UNIFORM_TEXEL_BUFFER:
+									{
+										TBufferSRVBindings.Add(Binding);
+										break;
+									}
+									case SPV_REFLECT_DESCRIPTOR_TYPE_STORAGE_BUFFER:
+									{
+										SBufferSRVBindings.Add(Binding);
+										break;
+									}
+									default:
+									{
+										// check(false);
+										break;
+									}
+								}
+								break;
+							}
+							case SPV_REFLECT_RESOURCE_FLAG_UAV:
+							{
+								switch(Binding->descriptor_type)
+								{
+									case SPV_REFLECT_DESCRIPTOR_TYPE_STORAGE_IMAGE:
+									{
+										TextureUAVBindings.Add(Binding);
+										break;
+									}
+									case SPV_REFLECT_DESCRIPTOR_TYPE_STORAGE_TEXEL_BUFFER:
+									{
+										TBufferUAVBindings.Add(Binding);
+										break;
+									}
+									case SPV_REFLECT_DESCRIPTOR_TYPE_STORAGE_BUFFER:
+									{
+										SBufferUAVBindings.Add(Binding);
+										break;
+									}
+									default:
+									{
+										// check(false);
+										break;
+									}
+								}
+								break;
+							}
+							default:
+							{
+								// check(false);
+								break;
+							}
+						}
+					}
+				}
+				
+				for (auto const& Binding : TBufferUAVBindings)
+				{
+					check(UAVIndices);
+					uint32 Index = FPlatformMath::CountTrailingZeros(UAVIndices);
+					
+					// UAVs always claim all slots so we don't have conflicts as D3D expects 0-7
+					BufferIndices &= ~(1 << Index);
+					TextureIndices &= ~(1 << Index);
+					UAVIndices &= ~(1 << Index);
+					
+					OutputData.TypedUAVs |= (1 << Index);
+					OutputData.TypedBuffers |= (1 << Index);
+					
+					UAVString += FString::Printf(TEXT("%s%s(%u:%u)"), UAVString.Len() ? TEXT(",") : TEXT(""), UTF8_TO_TCHAR(Binding->name), Index, 1);
+					
+					SPVRResult = Reflection.ChangeDescriptorBindingNumbers(Binding, Index);
+					check(SPVRResult == SPV_REFLECT_RESULT_SUCCESS);
+				}
+				
+				for (auto const& Binding : SBufferUAVBindings)
+				{
+					check(UAVIndices);
+					uint32 Index = FPlatformMath::CountTrailingZeros(UAVIndices);
+					
+					// UAVs always claim all slots so we don't have conflicts as D3D expects 0-7
+					BufferIndices &= ~(1 << Index);
+					TextureIndices &= ~(1 << Index);
+					UAVIndices &= ~(1 << Index);
+					
+					OutputData.InvariantBuffers |= (1 << Index);
+					
+					UAVString += FString::Printf(TEXT("%s%s(%u:%u)"), UAVString.Len() ? TEXT(",") : TEXT(""), UTF8_TO_TCHAR(Binding->name), Index, 1);
+					
+					SPVRResult = Reflection.ChangeDescriptorBindingNumbers(Binding, Index);
+					check(SPVRResult == SPV_REFLECT_RESULT_SUCCESS);
+				}
+				
+				for (auto const& Binding : TextureUAVBindings)
+				{
+					check(UAVIndices);
+					uint32 Index = FPlatformMath::CountTrailingZeros(UAVIndices);
+					
+					// UAVs always claim all slots so we don't have conflicts as D3D expects 0-7
+					// For texture2d this allows us to emulate atomics with buffers
+					BufferIndices &= ~(1 << Index);
+					TextureIndices &= ~(1 << Index);
+					UAVIndices &= ~(1 << Index);
+					
+					UAVString += FString::Printf(TEXT("%s%s(%u:%u)"), UAVString.Len() ? TEXT(",") : TEXT(""), UTF8_TO_TCHAR(Binding->name), Index, 1);
+					
+					SPVRResult = Reflection.ChangeDescriptorBindingNumbers(Binding, Index);
+					check(SPVRResult == SPV_REFLECT_RESULT_SUCCESS);
+				}
+				
+				for (auto const& Binding : TBufferSRVBindings)
+				{
+					check(TextureIndices);
+					uint32 Index = FPlatformMath::CountTrailingZeros(TextureIndices);
+					
+					// No support for 3-component types in dxc/SPIRV/MSL - need to expose my workarounds there too
+					// BufferIndices &= ~(1 << Index);
+					
+					TextureIndices &= ~(1 << Index);
+					
+					OutputData.TypedBuffers |= (1 << Index);
+					
+					SRVString += FString::Printf(TEXT("%s%s(%u:%u)"), SRVString.Len() ? TEXT(",") : TEXT(""), UTF8_TO_TCHAR(Binding->name), Index, 1);
+					
+					SPVRResult = Reflection.ChangeDescriptorBindingNumbers(Binding, Index);
+					check(SPVRResult == SPV_REFLECT_RESULT_SUCCESS);
+				}
+				
+				for (auto const& Binding : SBufferSRVBindings)
+				{
+					check(BufferIndices);
+					uint32 Index = FPlatformMath::CountTrailingZeros(BufferIndices);
+					
+					BufferIndices &= ~(1 << Index);
+					
+					OutputData.InvariantBuffers |= (1 << Index);
+					
+					SRVString += FString::Printf(TEXT("%s%s(%u:%u)"), SRVString.Len() ? TEXT(",") : TEXT(""), UTF8_TO_TCHAR(Binding->name), Index, 1);
+					
+					SPVRResult = Reflection.ChangeDescriptorBindingNumbers(Binding, Index);
+					check(SPVRResult == SPV_REFLECT_RESULT_SUCCESS);
+				}
+				
+				for (auto const& Binding : UniformBindings)
+				{
+					check(BufferIndices);
+					uint32 Index = FPlatformMath::CountTrailingZeros(BufferIndices);
+					BufferIndices &= ~(1 << Index);
+					
+					OutputData.ConstantBuffers |= (1 << Index);
+					
+					// Global uniform buffer - handled specially as we care about the internal layout
+					if (strstr(Binding->name, "$Globals"))
+					{
+						TCBDMARangeMap CBRanges;
+						GLOString = FString::Printf(TEXT("Globals(%u): "), Index);
+						
+						FString MbrString;
+						for (uint32 i = 0; i < Binding->block.member_count; i++)
+						{
+							SpvReflectBlockVariable& member = Binding->block.members[i];
+							uint32 MbrOffset = member.absolute_offset / sizeof(float);
+							uint32 MbrSize = member.size / sizeof(float);
+							
+							MbrString += FString::Printf(TEXT("%s%s(%u,%u)"), MbrString.Len() ? TEXT(",") : TEXT(""), UTF8_TO_TCHAR(member.name), MbrOffset, MbrSize);
+							
+							unsigned DestCBPrecision = TEXT('h');
+							
+//							auto const type = *member.type_description;
+//							uint32_t masked_type = type.type_flags & 0xF;
+//							
+//							switch (masked_type) {
+//								default: checkf(false, TEXT("unsupported component type %d"), masked_type); break;
+//								case SPV_REFLECT_TYPE_FLAG_BOOL  : DestCBPrecision = TEXT('u'); break;
+//								case SPV_REFLECT_TYPE_FLAG_INT   : DestCBPrecision = (type.traits.numeric.scalar.signedness ? TEXT('i') : TEXT('u')); break;
+//								case SPV_REFLECT_TYPE_FLAG_FLOAT : DestCBPrecision = (type.traits.numeric.scalar.width == 32 ? TEXT('h') : TEXT('m')); break;
+//							}
+							
+							unsigned SourceOffset = MbrOffset;
+							unsigned DestOffset = MbrOffset;
+							unsigned DestSize = MbrSize;
+							unsigned DestCBIndex = 0;
+							InsertRange(CBRanges, Index, SourceOffset, DestSize, DestCBIndex, DestCBPrecision, DestOffset);
+						}
+						
+						for (auto Iter = CBRanges.begin(); Iter != CBRanges.end(); ++Iter)
+						{
+							TDMARangeList& List = Iter->second;
+							for (auto IterList = List.begin(); IterList != List.end(); ++IterList)
+							{
+								check(IterList->DestCBIndex == 0);
+								PAKString += FString::Printf(TEXT("%s%u:%u-%c:%u:%u"), PAKString.Len() ? TEXT(",") : TEXT(""), IterList->SourceCB, IterList->SourceOffset, IterList->DestCBPrecision, IterList->DestOffset, IterList->Size);
+							}
+						}
+						
+						GLOString += MbrString;
+					}
+					else
+					{
+						// Regular uniform buffer - we only care about the binding index
+						UBOString += FString::Printf(TEXT("%s%s(%u)"), UBOString.Len() ? TEXT(",") : TEXT(""), UTF8_TO_TCHAR(Binding->name), Index);
+					}
+					
+					SPVRResult = Reflection.ChangeDescriptorBindingNumbers(Binding, Index);
+					check(SPVRResult == SPV_REFLECT_RESULT_SUCCESS);
+				}
+				
+				for (auto const& Binding : TextureSRVBindings)
+				{
+					check(TextureIndices);
+					uint32 Index = FPlatformMath::CountTrailingZeros64(TextureIndices);
+					TextureIndices &= ~(1 << Index);
+					
+					SRVString += FString::Printf(TEXT("%s%s(%u:%u)"), SRVString.Len() ? TEXT(",") : TEXT(""), UTF8_TO_TCHAR(Binding->name), Index, 1);
+					
+					SPVRResult = Reflection.ChangeDescriptorBindingNumbers(Binding, Index);
+					check(SPVRResult == SPV_REFLECT_RESULT_SUCCESS);
+				}
+				
+				for (auto const& Binding : SamplerBindings)
+				{
+					check(SamplerIndices);
+					uint32 Index = FPlatformMath::CountTrailingZeros(SamplerIndices);
+					SamplerIndices &= ~(1 << Index);
+					
+					SMPString += FString::Printf(TEXT("%s%u:%s"), SMPString.Len() ? TEXT(",") : TEXT(""), Index, UTF8_TO_TCHAR(Binding->name));
+					
+					SPVRResult = Reflection.ChangeDescriptorBindingNumbers(Binding, Index);
+					check(SPVRResult == SPV_REFLECT_RESULT_SUCCESS);
+				}
+			}
+			
+			if (Frequency == HSF_PixelShader)
+			{
+				Count = 0;
+				SPVRResult = Reflection.EnumerateOutputVariables(&Count, nullptr);
+				check(SPVRResult == SPV_REFLECT_RESULT_SUCCESS);
+				OutputVars.SetNum(Count);
+				SPVRResult = Reflection.EnumerateOutputVariables(&Count, OutputVars.GetData());
+				check(SPVRResult == SPV_REFLECT_RESULT_SUCCESS);
+				if (Count > 0)
+				{
+					for (auto const& Var : OutputVars)
+					{
+						if (Var->storage_class == SpvStorageClassOutput && Var->built_in == -1 && strstr(Var->name, "SV_Target"))
+						{
+							FString TypeQualifier;
+							
+							auto const type = *Var->type_description;
+							uint32_t masked_type = type.type_flags & 0xF;
+							
+							switch (masked_type) {
+								default: checkf(false, TEXT("unsupported component type %d"), masked_type); break;
+								case SPV_REFLECT_TYPE_FLAG_BOOL  : TypeQualifier = TEXT("b"); break;
+								case SPV_REFLECT_TYPE_FLAG_INT   : TypeQualifier = (type.traits.numeric.scalar.signedness ? TEXT("i") : TEXT("u")); break;
+								case SPV_REFLECT_TYPE_FLAG_FLOAT : TypeQualifier = (type.traits.numeric.scalar.width == 32 ? TEXT("f") : TEXT("h")); break;
+							}
+							
+							if (type.type_flags & SPV_REFLECT_TYPE_FLAG_MATRIX)
+							{
+								TypeQualifier += FString::Printf(TEXT("%d%d"), type.traits.numeric.matrix.row_count, type.traits.numeric.matrix.column_count);
+							}
+							else if (type.type_flags & SPV_REFLECT_TYPE_FLAG_VECTOR)
+							{
+								TypeQualifier += FString::Printf(TEXT("%d"), type.traits.numeric.vector.component_count);
+							}
+							else
+							{
+								TypeQualifier += TEXT("1");
+							}
+							
+							OUTString += FString::Printf(TEXT("%s%s:SV_Target%d"), OUTString.Len() ? TEXT(",") : TEXT(""), *TypeQualifier, Var->location);
+						}
+					}
+				}
+			}
+			
+			if (Frequency == HSF_VertexShader)
+			{
+				Count = 0;
+				SPVRResult = Reflection.EnumerateInputVariables(&Count, nullptr);
+				check(SPVRResult == SPV_REFLECT_RESULT_SUCCESS);
+				InputVars.SetNum(Count);
+				SPVRResult = Reflection.EnumerateInputVariables(&Count, InputVars.GetData());
+				check(SPVRResult == SPV_REFLECT_RESULT_SUCCESS);
+				if (Count > 0)
+				{
+					uint32 AssignedInputs = 0;
+					
+					for (auto const& Var : InputVars)
+					{
+						if (Var->storage_class == SpvStorageClassInput && Var->built_in == -1)
+						{
+							unsigned Location = Var->location;
+							unsigned SemanticIndex = Location;
+							check(Var->semantic);
+							unsigned i = (unsigned)strlen(Var->semantic);
+							check(i);
+							while (isdigit(Var->semantic[i-1]))
+							{
+								i--;
+							}
+							if (i < strlen(Var->semantic))
+							{
+								SemanticIndex = (unsigned)atoi(Var->semantic + i);
+								if (Location != SemanticIndex)
+								{
+									Location = SemanticIndex;
+								}
+							}
+							
+							while ((1 << Location) & AssignedInputs)
+							{
+								Location++;
+							}
+							
+							if (Location != Var->location)
+							{
+								SPVRResult = Reflection.ChangeInputVariableLocation(Var, Location);
+								check(SPVRResult == SPV_REFLECT_RESULT_SUCCESS);
+							}
+							
+							uint32 ArrayCount = 1;
+							for (uint32 Dim = 0; Dim < Var->array.dims_count; Dim++)
+							{
+								ArrayCount *= Var->array.dims[Dim];
+							}
+							
+							FString TypeQualifier;
+
+							auto const type = *Var->type_description;
+							uint32_t masked_type = type.type_flags & 0xF;
+							
+							switch (masked_type) {
+								default: checkf(false, TEXT("unsupported component type %d"), masked_type); break;
+								case SPV_REFLECT_TYPE_FLAG_BOOL  : TypeQualifier = TEXT("b"); break;
+								case SPV_REFLECT_TYPE_FLAG_INT   : TypeQualifier = (type.traits.numeric.scalar.signedness ? TEXT("i") : TEXT("u")); break;
+								case SPV_REFLECT_TYPE_FLAG_FLOAT : TypeQualifier = (type.traits.numeric.scalar.width == 32 ? TEXT("f") : TEXT("h")); break;
+							}
+							
+							if (type.type_flags & SPV_REFLECT_TYPE_FLAG_MATRIX)
+							{
+								TypeQualifier += FString::Printf(TEXT("%d%d"), type.traits.numeric.matrix.row_count, type.traits.numeric.matrix.column_count);
+							}
+							else if (type.type_flags & SPV_REFLECT_TYPE_FLAG_VECTOR)
+							{
+								TypeQualifier += FString::Printf(TEXT("%d"), type.traits.numeric.vector.component_count);
+							}
+							else
+							{
+								TypeQualifier += TEXT("1");
+							}
+							
+							for (uint32 j = 0; j < ArrayCount; j++)
+							{
+								AssignedInputs |= (1 << (Location + j));
+								
+								INPString += FString::Printf(TEXT("%s%s:in_ATTRIBUTE%d"), INPString.Len() ? TEXT(",") : TEXT(""), *TypeQualifier, (Location + j));
+							}
+						}
+					}
+				}
+			}
+			
+			MetaData += TEXT("// Compiled by ShaderConductor\n");
+			if (INPString.Len())
+			{
+				MetaData += FString::Printf(TEXT("// @Inputs: %s\n"), *INPString);
+			}
+			if (OUTString.Len())
+			{
+				MetaData += FString::Printf(TEXT("// @Outputs: %s\n"), *OUTString);
+			}
+			if (UBOString.Len())
+			{
+				MetaData += FString::Printf(TEXT("// @UniformBlocks: %s\n"), *UBOString);
+			}
+			if (GLOString.Len())
+			{
+				MetaData += FString::Printf(TEXT("// @PackedUB: %s\n"), *GLOString);
+			}
+			if (PAKString.Len())
+			{
+				MetaData += FString::Printf(TEXT("// @PackedUBGlobalCopies: %s\n"), *PAKString);
+			}
+			if (SRVString.Len())
+			{
+				MetaData += FString::Printf(TEXT("// @Samplers: %s\n"), *SRVString);
+			}
+			if (SMPString.Len())
+			{
+				MetaData += FString::Printf(TEXT("// @SamplerStates: %s\n"), *SMPString);
+			}
+			if (UAVString.Len())
+			{
+				MetaData += FString::Printf(TEXT("// @UAVs: %s\n"), *UAVString);
+			}
+			if (WKGString.Len())
+			{
+				MetaData += FString::Printf(TEXT("// @NumThreads: %s\n"), *WKGString);
+			}
+			MetaData += FString::Printf(TEXT("\n\n"), *WKGString);
+			
+			ShaderConductor::Blob* OldData = Results.target;
+			Results.target = ShaderConductor::CreateBlob(Reflection.GetCode(), Reflection.GetCodeSize());
+			
+			FArchive* FileWriter = IFileManager::Get().CreateFileWriter(*(Input.DumpDebugInfoPath / FPaths::GetBaseFilename(Input.GetSourceFilename()) + TEXT(".dxc.spirv")));
+			if (FileWriter)
+			{
+				FileWriter->Serialize((void*)Reflection.GetCode(), Reflection.GetCodeSize());
+				FileWriter->Close();
+				delete FileWriter;
+			}
+			
+			ShaderConductor::DestroyBlob(OldData);
+		}
+		
+		if (Result)
+		{
+			ShaderConductor::DestroyBlob(Results.errorWarningMsg);
+			Results.errorWarningMsg = nullptr;
+			
+			TargetDesc.language = ShaderConductor::ShadingLanguage::Msl;
+			
+			switch (VersionEnum)
+			{
+				case 6:
+				case 5:
+				case 4:
+				{
+					TargetDesc.version = "20100";
+					break;
+				}
+				case 3:
+				{
+					TargetDesc.version = "20000";
+					break;
+				}
+				case 2:
+				{
+					TargetDesc.version = "10200";
+					break;
+				}
+				case 1:
+				{
+					TargetDesc.version = "10100";
+					break;
+				}
+				case 0:
+				default:
+				{
+					TargetDesc.version = "10000";
+					break;
+				}
+			}
+			
+			ShaderConductor::MacroDefine Defines[3] = {{"texel_buffer_texture_width", "0"}, {nullptr, nullptr}, {nullptr, nullptr}};
+			TargetDesc.numOptions = 1;
+			TargetDesc.options = &Defines[0];
+			switch(Semantics)
+			{
+				case EMetalGPUSemanticsImmediateDesktop:
+					TargetDesc.platform = "macOS";
+					break;
+				case EMetalGPUSemanticsTBDRDesktop:
+					TargetDesc.platform = "iOS";
+					Defines[1] = { "ios_support_base_vertex_instance", "1" };
+					Defines[2] = { "ios_use_framebuffer_fetch_subpasses", "1" };
+					TargetDesc.numOptions += 2;
+					break;
+				case EMetalGPUSemanticsMobile:
+				default:
+					TargetDesc.platform = "iOS";
+					Defines[1] = { "ios_use_framebuffer_fetch_subpasses", "1" };
+					TargetDesc.numOptions += 1;
+					break;
+			}
+			
+			ShaderConductor::Blob* IRBlob = Results.target;
+			
+			Results = ShaderConductor::Compiler::ConvertBinary(Results, SourceDesc, TargetDesc);
+			
+			ShaderConductor::DestroyBlob(IRBlob);
+		}
+		
+		Result = (Results.hasError) ? 0 : 1;
+		if (!Results.hasError)
+		{
+			MetalSource = TCHAR_TO_UTF8(*MetaData);
+			MetalSource += std::string((const char*)Results.target->Data(), Results.target->Size());
+			
+			CRCLen = MetalSource.length();
+			CRC = FCrc::MemCrc_DEPRECATED(MetalSource.c_str(), CRCLen);
+			
+			ANSICHAR MainCRC[25];
+			int32 NewLen = FCStringAnsi::Snprintf(MainCRC, 25, "Main_%0.8x_%0.8x(", CRCLen, CRC);
+			
+			std::string MainEntryPoint = EntryPointName + "(";
+			size_t Pos;
+			do
+			{
+				Pos = MetalSource.find(MainEntryPoint);
+				if (Pos != std::string::npos)
+					MetalSource.replace(Pos, MainEntryPoint.length(), MainCRC);
+			} while(Pos != std::string::npos);
+		}
+		
+		if (Results.errorWarningMsg)
+		{
+			MetalErrors = ANSI_TO_TCHAR((char const*)Results.errorWarningMsg->Data());
+			ShaderConductor::DestroyBlob(Results.errorWarningMsg);
+			Results.errorWarningMsg = nullptr;
+		}
+		
+		if (Results.target)
+		{
+			ShaderConductor::DestroyBlob(Results.target);
+			Results.target = nullptr;
+		}
+		
+		if (RewriteBlob)
+		{
+			ShaderConductor::DestroyBlob(RewriteBlob);
+		}
 	}
+#endif
 
 	if (bDumpDebugInfo)
 	{
-		if (SourceLen > 0u)
+		if (MetalSource.length() > 0u)
 		{
 			FArchive* FileWriter = IFileManager::Get().CreateFileWriter(*(Input.DumpDebugInfoPath / FPaths::GetBaseFilename(Input.GetSourceFilename()) + TEXT(".metal")));
 			if (FileWriter)
 			{
-				FileWriter->Serialize(MetalShaderSource, SourceLen + 1);
+				FileWriter->Serialize(const_cast<char*>(MetalSource.c_str()), MetalSource.length());
 				FileWriter->Close();
 				delete FileWriter;
 			}
@@ -450,21 +1211,21 @@ bool FMetalShaderOutputCooker::Build(TArray<uint8>& OutData)
 	if (Result != 0)
 	{
 		Output.Target = Input.Target;
-		BuildMetalShaderOutput(Output, Input, GUIDHash, CCFlags, MetalShaderSource, SourceLen, CRCLen, CRC, VersionEnum, *Standard, *MinOSVersion, TypeMode, Output.Errors, Attribs, MetalBackEnd.TypedBuffers, MetalBackEnd.InvariantBuffers, MetalBackEnd.TypedUAVs, MetalBackEnd.ConstantBuffers, MetalBackEnd.TypedBufferFormats, bAllowFastIntriniscs);
+		BuildMetalShaderOutput(Output, Input, GUIDHash, CCFlags, MetalSource.c_str(), MetalSource.length(), CRCLen, CRC, VersionEnum, *Standard, *MinOSVersion, TypeMode, Output.Errors, Attribs, OutputData.TypedBuffers, OutputData.InvariantBuffers, OutputData.TypedUAVs, OutputData.ConstantBuffers, OutputData.TypedBufferFormats, bAllowFastIntriniscs);
 
 		FMemoryWriter Ar(OutData);
 		Ar << Output;
 	}
 	else
 	{
-		FString Tmp = ANSI_TO_TCHAR(ErrorLog);
 		TArray<FString> ErrorLines;
-		Tmp.ParseIntoArray(ErrorLines, TEXT("\n"), true);
+		MetalErrors.ParseIntoArray(ErrorLines, TEXT("\n"), true);
 		bool const bDirectCompile = FParse::Param(FCommandLine::Get(), TEXT("directcompile"));
 		for (int32 LineIndex = 0; LineIndex < ErrorLines.Num(); ++LineIndex)
 		{
 			const FString& Line = ErrorLines[LineIndex];
-			CrossCompiler::ParseHlslccError(Output.Errors, Line, bDirectCompile);
+			Output.Errors.Add(FShaderCompilerError(*Line));
+			CrossCompiler::ParseHlslccError(Output.Errors, Line, true);
 			if (bDirectCompile)
 			{
 				UE_LOG(LogShaders, Error, TEXT("%s"), *Line);
@@ -472,14 +1233,5 @@ bool FMetalShaderOutputCooker::Build(TArray<uint8>& OutData)
 		}
 	}
 
-	if (MetalShaderSource)
-	{
-		free(MetalShaderSource);
-	}
-	if (ErrorLog)
-	{
-		free(ErrorLog);
-	}
-
 	return Output.bSucceeded;
 }
diff --git a/Engine/Source/Developer/Apple/MetalShaderFormat/Private/MetalShaderCompiler.cpp b/Engine/Source/Developer/Apple/MetalShaderFormat/Private/MetalShaderCompiler.cpp
--- a/Engine/Source/Developer/Apple/MetalShaderFormat/Private/MetalShaderCompiler.cpp
+++ b/Engine/Source/Developer/Apple/MetalShaderFormat/Private/MetalShaderCompiler.cpp
@@ -1134,7 +1134,7 @@ void BuildMetalShaderOutput(
 
 	// Then the list of outputs.
 	static const FString TargetPrefix = "FragColor";
-	static const FString GL_FragDepth = "FragDepth";
+	static const FString TargetPrefix2 = "SV_Target";
 	// Only outputs for pixel shaders must be tracked.
 	if (Frequency == SF_Pixel)
 	{
@@ -1146,12 +1146,18 @@ void BuildMetalShaderOutput(
 				uint8 TargetIndex = ParseNumber(*Output.Name + TargetPrefix.Len());
 				Header.Bindings.InOutMask |= (1 << TargetIndex);
 			}
-			// Handle depth writes.
-			else if (Output.Name.Equals(GL_FragDepth))
+			else if (Output.Name.StartsWith(TargetPrefix2))
 			{
-				Header.Bindings.InOutMask |= 0x8000;
+				uint8 TargetIndex = ParseNumber(*Output.Name + TargetPrefix2.Len());
+				Header.Bindings.InOutMask |= (1 << TargetIndex);
 			}
         }
+		
+		// For fragment shaders that discard but don't output anything we need at least a depth-stencil surface, so we need a way to validate this at runtime.
+		if (FCStringAnsi::Strstr(USFSource, "[[ depth(") != nullptr || FCStringAnsi::Strstr(USFSource, "[[depth(") != nullptr)
+		{
+			Header.Bindings.InOutMask |= 0x8000;
+		}
         
         // For fragment shaders that discard but don't output anything we need at least a depth-stencil surface, so we need a way to validate this at runtime.
         if (FCStringAnsi::Strstr(USFSource, "discard_fragment()") != nullptr)
@@ -1937,8 +1943,7 @@ void CompileShader_Metal(const FShaderCompilerInput& _Input,FShaderCompilerOutpu
 		AdditionalDefines.SetDefine(TEXT("MAC"), 1);
 	}
 	
-	AdditionalDefines.SetDefine(TEXT("COMPILER_HLSLCC"), 1 );
-	AdditionalDefines.SetDefine(TEXT("row_major"), TEXT(""));
+	AdditionalDefines.SetDefine(TEXT("COMPILER_HLSLCC"), 1);
 	AdditionalDefines.SetDefine(TEXT("COMPILER_METAL"), 1);
 
 	static FName NAME_SF_METAL(TEXT("SF_METAL"));
@@ -1963,6 +1968,11 @@ void CompileShader_Metal(const FShaderCompilerInput& _Input,FShaderCompilerOutpu
         }
     }
 	
+	// The new compiler is only available on Mac or Windows for the moment.
+#if !(PLATFORM_MAC || PLATFORM_WINDOWS)
+	VersionEnum = FMath::Min(VersionEnum, (uint8)5);
+#endif
+	
 	bool bAppleTV = (Input.ShaderFormat == NAME_SF_METAL_TVOS || Input.ShaderFormat == NAME_SF_METAL_MRT_TVOS);
     if (Input.ShaderFormat == NAME_SF_METAL || Input.ShaderFormat == NAME_SF_METAL_TVOS)
 	{
@@ -2027,11 +2037,17 @@ void CompileShader_Metal(const FShaderCompilerInput& _Input,FShaderCompilerOutpu
 		return;
 	}
 	
+	if (VersionEnum < 6)
+	{
+		AdditionalDefines.SetDefine(TEXT("row_major"), TEXT(""));
+	}
+	
     EMetalTypeBufferMode TypeMode = EMetalTypeBufferModeRaw;
 	FString MinOSVersion;
 	FString StandardVersion;
 	switch(VersionEnum)
     {
+		case 6:
         case 5:
             // Enable full SM5 feature support so tessellation & fragment UAVs compile
             TypeMode = EMetalTypeBufferModeTB;
@@ -2241,12 +2257,12 @@ void CompileShader_Metal(const FShaderCompilerInput& _Input,FShaderCompilerOutpu
 
 
 	// This requires removing the HLSLCC_NoPreprocess flag later on!
-    if (VersionEnum < 5)
+    if (VersionEnum < 5 || VersionEnum == 6)
     {
         RemoveUniformBuffersFromSource(Input.Environment, PreprocessedShader);
     }
 
-	uint32 CCFlags = HLSLCC_NoPreprocess | HLSLCC_PackUniformsIntoUniformBufferWithNames | HLSLCC_FixAtomicReferences | HLSLCC_KeepSamplerAndImageNames;
+	uint32 CCFlags = HLSLCC_NoPreprocess | HLSLCC_PackUniformsIntoUniformBufferWithNames | HLSLCC_FixAtomicReferences | HLSLCC_RetainSizes | HLSLCC_KeepSamplerAndImageNames;
 	if (!bDirectCompile || UE_BUILD_DEBUG)
 	{
 		// Validation is expensive - only do it when compiling directly for debugging
diff --git a/Engine/Source/Runtime/Apple/MetalRHI/Private/MetalCommandQueue.cpp b/Engine/Source/Runtime/Apple/MetalRHI/Private/MetalCommandQueue.cpp
--- a/Engine/Source/Runtime/Apple/MetalRHI/Private/MetalCommandQueue.cpp
+++ b/Engine/Source/Runtime/Apple/MetalRHI/Private/MetalCommandQueue.cpp
@@ -37,8 +37,8 @@ FMetalCommandQueue::FMetalCommandQueue(mtlpp::Device InDevice, uint32 const MaxN
 	int32 MinShaderVersion = 3;
     const TCHAR* const Settings = TEXT("/Script/MacTargetPlatform.MacTargetSettings");
 #else
-	int32 DefaultMaxShaderVersion = 0;
-	int32 MinShaderVersion = 0;
+	int32 DefaultMaxShaderVersion = 2;
+	int32 MinShaderVersion = 2;
     const TCHAR* const Settings = TEXT("/Script/IOSRuntimeSettings.IOSRuntimeSettings");
 #endif
     if(!GConfig->GetInt(Settings, TEXT("MaxShaderLanguageVersion"), MaxShaderVersion, GEngineIni))
@@ -229,7 +229,7 @@ FMetalCommandQueue::FMetalCommandQueue(mtlpp::Device InDevice, uint32 const MaxN
 			{
 				Features |= EMetalFeaturesTextureBuffers;
             }
-            if (MaxShaderVersion >= 5)
+            if (MaxShaderVersion == 5)
             {
                 Features |= EMetalFeaturesIABs;
             }
diff --git a/Engine/Source/Runtime/Apple/MetalRHI/Private/MetalShaders.cpp b/Engine/Source/Runtime/Apple/MetalRHI/Private/MetalShaders.cpp
--- a/Engine/Source/Runtime/Apple/MetalRHI/Private/MetalShaders.cpp
+++ b/Engine/Source/Runtime/Apple/MetalRHI/Private/MetalShaders.cpp
@@ -410,6 +410,8 @@ void TMetalBaseShader<BaseResourceType, ShaderType>::Init(const TArray<uint8>& I
 			mtlpp::LanguageVersion MetalVersion;
 			switch(Header.Version)
 			{
+				case 6:
+				case 5:
 				case 4:
 					MetalVersion = mtlpp::LanguageVersion::Version2_1;
 					break;
diff --git a/Engine/Source/Runtime/Engine/Private/ShaderCompiler/ShaderCompiler.cpp b/Engine/Source/Runtime/Engine/Private/ShaderCompiler/ShaderCompiler.cpp
--- a/Engine/Source/Runtime/Engine/Private/ShaderCompiler/ShaderCompiler.cpp
+++ b/Engine/Source/Runtime/Engine/Private/ShaderCompiler/ShaderCompiler.cpp
@@ -2732,7 +2732,7 @@ ENGINE_API void GenerateInstancedStereoCode(FString& Result)
 	Result += "}\r\n";
 	
 	// ResolveView definition for metal, this allows us to change the branch to a conditional move in the cross compiler
-	Result += "#if COMPILER_METAL\r\n";
+	Result += "#if COMPILER_METAL && MAX_SHADER_LANGUAGE_VERSION < 6\r\n";
 	Result += "ViewState ResolveView(uint ViewIndex)\r\n";
 	Result += "{\r\n";
 	Result += "\tViewState Result;\r\n";
