Disable ShaderConductor while I work out why it won't load.

diff --git a/Engine/Source/Developer/Apple/MetalShaderFormat/MetalShaderFormat.Build.cs b/Engine/Source/Developer/Apple/MetalShaderFormat/MetalShaderFormat.Build.cs
--- a/Engine/Source/Developer/Apple/MetalShaderFormat/MetalShaderFormat.Build.cs
+++ b/Engine/Source/Developer/Apple/MetalShaderFormat/MetalShaderFormat.Build.cs
@@ -28,7 +28,7 @@ public class MetalShaderFormat : ModuleRules
 		AddEngineThirdPartyPrivateStaticDependencies(Target, 
 			"HLSLCC"
 			);
-
+/*
 		if (Target.Platform == UnrealTargetPlatform.Mac || Target.Platform == UnrealTargetPlatform.Win64)
 		{
 			AddEngineThirdPartyPrivateStaticDependencies(Target,
@@ -36,7 +36,7 @@ public class MetalShaderFormat : ModuleRules
 				"SPIRVReflect"
 			);
 		}
-			
+*/
 		DynamicallyLoadedModuleNames.AddRange(
 			new string[]
 			{ 
diff --git a/Engine/Source/Developer/Apple/MetalShaderFormat/Private/MetalDerivedData.cpp b/Engine/Source/Developer/Apple/MetalShaderFormat/Private/MetalDerivedData.cpp
--- a/Engine/Source/Developer/Apple/MetalShaderFormat/Private/MetalDerivedData.cpp
+++ b/Engine/Source/Developer/Apple/MetalShaderFormat/Private/MetalDerivedData.cpp
@@ -13,12 +13,14 @@
 #include "MetalBackend.h"
 #include "Serialization/MemoryReader.h"
 
+#if 0
 #if PLATFORM_MAC || PLATFORM_WINDOWS
 THIRD_PARTY_INCLUDES_START
 #include "ShaderConductor/ShaderConductor.hpp"
 #include "spirv_reflect.h"
 THIRD_PARTY_INCLUDES_END
 #endif
+#endif
 
 extern bool ExecRemoteProcess(const TCHAR* Command, const TCHAR* Params, int32* OutReturnCode, FString* OutStdOut, FString* OutStdErr);
 extern FString GetXcodePath();
@@ -425,8 +427,10 @@ bool FMetalShaderOutputCooker::Build(TArray<uint8>& OutData)
 	bool bUsingTessellation = ((UsingTessellationDefine != nullptr && FString("1") == *UsingTessellationDefine && Frequency == HSF_VertexShader) || Frequency == HSF_HullShader || Frequency == HSF_DomainShader);
 	
 	// Its going to take a while to get dxc+SPIRV working...
+#if 0
 #if PLATFORM_MAC || PLATFORM_WINDOWS
 	if (VersionEnum < 6)
+#endif
 #endif
 	{
 		FMetalCodeBackend MetalBackEnd(Attribs, CCFlags, MetalCompilerTarget, VersionEnum, Semantics, TypeMode, MaxUnrollLoops, bZeroInitialise, bBoundsChecks, bAllowFastIntriniscs, bForceInvariance, bSwizzleSample);
@@ -484,6 +488,7 @@ bool FMetalShaderOutputCooker::Build(TArray<uint8>& OutData)
 			}
 		}
 	}
+#if 0
 #if PLATFORM_MAC || PLATFORM_WINDOWS
 	else
 	{
@@ -1192,6 +1197,7 @@ bool FMetalShaderOutputCooker::Build(TArray<uint8>& OutData)
 			ShaderConductor::DestroyBlob(RewriteBlob);
 		}
 	}
+#endif
 #endif
 
 	if (bDumpDebugInfo)
