Fix Zen's particles using Metal ES 3.1 when compiled with ShaderConductor - we need to bind buffers & textures for Buffer<> types to match hlslcc's behaviour.

diff --git a/Engine/Source/Developer/Apple/MetalShaderFormat/Private/MetalDerivedData.cpp b/Engine/Source/Developer/Apple/MetalShaderFormat/Private/MetalDerivedData.cpp
--- a/Engine/Source/Developer/Apple/MetalShaderFormat/Private/MetalDerivedData.cpp
+++ b/Engine/Source/Developer/Apple/MetalShaderFormat/Private/MetalDerivedData.cpp
@@ -788,8 +788,7 @@ bool FMetalShaderOutputCooker::Build(TArray<uint8>& OutData)
 					uint32 Index = FPlatformMath::CountTrailingZeros(TextureIndices);
 					
 					// No support for 3-component types in dxc/SPIRV/MSL - need to expose my workarounds there too
-					// BufferIndices &= ~(1 << Index);
-					
+					BufferIndices &= ~(1 << Index);
 					TextureIndices &= ~(1 << Index);
 					
 					OutputData.TypedBuffers |= (1 << Index);
