Fix assignment of indices to UAVs that are only used in atomic operations when compilng with SC.

diff --git a/Engine/Source/Developer/Apple/MetalShaderFormat/Private/MetalDerivedData.cpp b/Engine/Source/Developer/Apple/MetalShaderFormat/Private/MetalDerivedData.cpp
--- a/Engine/Source/Developer/Apple/MetalShaderFormat/Private/MetalDerivedData.cpp
+++ b/Engine/Source/Developer/Apple/MetalShaderFormat/Private/MetalDerivedData.cpp
@@ -649,81 +649,93 @@ bool FMetalShaderOutputCooker::Build(TArray<uint8>& OutData)
 				// Which is necessary to match the D3D binding scheme.
 				for (auto const& Binding : Bindings)
 				{
-					if (Binding->accessed)
+					switch(Binding->resource_type)
 					{
-						switch(Binding->resource_type)
+						case SPV_REFLECT_RESOURCE_FLAG_CBV:
 						{
-							case SPV_REFLECT_RESOURCE_FLAG_CBV:
+							check(Binding->descriptor_type == SPV_REFLECT_DESCRIPTOR_TYPE_UNIFORM_BUFFER);
+							if (Binding->accessed)
 							{
-								check(Binding->descriptor_type == SPV_REFLECT_DESCRIPTOR_TYPE_UNIFORM_BUFFER);
 								UniformBindings.Add(Binding);
-								break;
 							}
-							case SPV_REFLECT_RESOURCE_FLAG_SAMPLER:
+							break;
+						}
+						case SPV_REFLECT_RESOURCE_FLAG_SAMPLER:
+						{
+							check(Binding->descriptor_type == SPV_REFLECT_DESCRIPTOR_TYPE_SAMPLER);
+							if (Binding->accessed)
 							{
-								check(Binding->descriptor_type == SPV_REFLECT_DESCRIPTOR_TYPE_SAMPLER);
 								SamplerBindings.Add(Binding);
-								break;
 							}
-							case SPV_REFLECT_RESOURCE_FLAG_SRV:
+							break;
+						}
+						case SPV_REFLECT_RESOURCE_FLAG_SRV:
+						{
+							switch(Binding->descriptor_type)
 							{
-								switch(Binding->descriptor_type)
+								case SPV_REFLECT_DESCRIPTOR_TYPE_SAMPLED_IMAGE:
 								{
-									case SPV_REFLECT_DESCRIPTOR_TYPE_SAMPLED_IMAGE:
+									if (Binding->accessed)
 									{
 										TextureSRVBindings.Add(Binding);
-										break;
 									}
-									case SPV_REFLECT_DESCRIPTOR_TYPE_UNIFORM_TEXEL_BUFFER:
+									break;
+								}
+								case SPV_REFLECT_DESCRIPTOR_TYPE_UNIFORM_TEXEL_BUFFER:
+								{
+									if (Binding->accessed)
 									{
 										TBufferSRVBindings.Add(Binding);
-										break;
 									}
-									case SPV_REFLECT_DESCRIPTOR_TYPE_STORAGE_BUFFER:
+									break;
+								}
+								case SPV_REFLECT_DESCRIPTOR_TYPE_STORAGE_BUFFER:
+								{
+									if (Binding->accessed)
 									{
 										SBufferSRVBindings.Add(Binding);
-										break;
-									}
-									default:
-									{
-										// check(false);
-										break;
 									}
+									break;
 								}
-								break;
-							}
-							case SPV_REFLECT_RESOURCE_FLAG_UAV:
-							{
-								switch(Binding->descriptor_type)
+								default:
 								{
-									case SPV_REFLECT_DESCRIPTOR_TYPE_STORAGE_IMAGE:
-									{
-										TextureUAVBindings.Add(Binding);
-										break;
-									}
-									case SPV_REFLECT_DESCRIPTOR_TYPE_STORAGE_TEXEL_BUFFER:
-									{
-										TBufferUAVBindings.Add(Binding);
-										break;
-									}
-									case SPV_REFLECT_DESCRIPTOR_TYPE_STORAGE_BUFFER:
-									{
-										SBufferUAVBindings.Add(Binding);
-										break;
-									}
-									default:
-									{
-										// check(false);
-										break;
-									}
+									// check(false);
+									break;
 								}
-								break;
 							}
-							default:
+							break;
+						}
+						case SPV_REFLECT_RESOURCE_FLAG_UAV:
+						{
+							switch(Binding->descriptor_type)
 							{
-								// check(false);
-								break;
+								case SPV_REFLECT_DESCRIPTOR_TYPE_STORAGE_IMAGE:
+								{
+									TextureUAVBindings.Add(Binding);
+									break;
+								}
+								case SPV_REFLECT_DESCRIPTOR_TYPE_STORAGE_TEXEL_BUFFER:
+								{
+									TBufferUAVBindings.Add(Binding);
+									break;
+								}
+								case SPV_REFLECT_DESCRIPTOR_TYPE_STORAGE_BUFFER:
+								{
+									SBufferUAVBindings.Add(Binding);
+									break;
+								}
+								default:
+								{
+									// check(false);
+									break;
+								}
 							}
+							break;
+						}
+						default:
+						{
+							// check(false);
+							break;
 						}
 					}
 				}
