Fix the Metal shader header for ShaderConductor compiled shaders.

diff --git a/Engine/Source/Developer/Apple/MetalShaderFormat/Private/MetalShaderCompiler.cpp b/Engine/Source/Developer/Apple/MetalShaderFormat/Private/MetalShaderCompiler.cpp
--- a/Engine/Source/Developer/Apple/MetalShaderFormat/Private/MetalShaderCompiler.cpp
+++ b/Engine/Source/Developer/Apple/MetalShaderFormat/Private/MetalShaderCompiler.cpp
@@ -1102,6 +1102,12 @@ void BuildMetalShaderOutput(
 			}
 		}
 		
+		if (Version == 6)
+		{
+			Header.Bindings.LinearBuffer = Header.Bindings.TypedBuffers;
+			Header.Bindings.TypedBuffers = 0;
+		}
+		
 		// Raw mode means all buffers are invariant
 		if (TypeMode == EMetalTypeBufferModeRaw)
 		{
