Fix the order of cross-compiler header order so that UAV info is parsed.

diff --git a/Engine/Source/Developer/Apple/MetalShaderFormat/Private/MetalDerivedData.cpp b/Engine/Source/Developer/Apple/MetalShaderFormat/Private/MetalDerivedData.cpp
--- a/Engine/Source/Developer/Apple/MetalShaderFormat/Private/MetalDerivedData.cpp
+++ b/Engine/Source/Developer/Apple/MetalShaderFormat/Private/MetalDerivedData.cpp
@@ -1077,14 +1077,14 @@ bool FMetalShaderOutputCooker::Build(TArray<uint8>& OutData)
 			{
 				MetaData += FString::Printf(TEXT("// @Samplers: %s\n"), *SRVString);
 			}
-			if (SMPString.Len())
-			{
-				MetaData += FString::Printf(TEXT("// @SamplerStates: %s\n"), *SMPString);
-			}
 			if (UAVString.Len())
 			{
 				MetaData += FString::Printf(TEXT("// @UAVs: %s\n"), *UAVString);
 			}
+			if (SMPString.Len())
+			{
+				MetaData += FString::Printf(TEXT("// @SamplerStates: %s\n"), *SMPString);
+			}
 			if (WKGString.Len())
 			{
 				MetaData += FString::Printf(TEXT("// @NumThreads: %s\n"), *WKGString);
@@ -1247,7 +1247,7 @@ bool FMetalShaderOutputCooker::Build(TArray<uint8>& OutData)
 		{
 			const FString& Line = ErrorLines[LineIndex];
 			Output.Errors.Add(FShaderCompilerError(*Line));
-			CrossCompiler::ParseHlslccError(Output.Errors, Line, true);
+			CrossCompiler::ParseHlslccError(Output.Errors, Line, false);
 			if (bDirectCompile)
 			{
 				UE_LOG(LogShaders, Error, TEXT("%s"), *Line);
