Fix static analysis warnings.

diff --git a/Engine/Source/Developer/Apple/MetalShaderFormat/Private/MetalDerivedData.cpp b/Engine/Source/Developer/Apple/MetalShaderFormat/Private/MetalDerivedData.cpp
--- a/Engine/Source/Developer/Apple/MetalShaderFormat/Private/MetalDerivedData.cpp
+++ b/Engine/Source/Developer/Apple/MetalShaderFormat/Private/MetalDerivedData.cpp
@@ -622,9 +622,9 @@ bool FMetalShaderOutputCooker::Build(TArray<uint8>& OutData)
 				uint32 WorkgroupSize[3];
 
 #if !PLATFORM_WINDOWS
-				size_t Found = sscanf(String, "numthreads(%d, %d, %d)", &WorkgroupSize[0], &WorkgroupSize[1], &WorkgroupSize[2]);
+				size_t Found = sscanf(String, "numthreads(%u, %u, %u)", &WorkgroupSize[0], &WorkgroupSize[1], &WorkgroupSize[2]);
 #else
-				size_t Found = sscanf_s(String, "numthreads(%d, %d, %d)", &WorkgroupSize[0], &WorkgroupSize[1], &WorkgroupSize[2]);
+				size_t Found = sscanf_s(String, "numthreads(%u, %u, %u)", &WorkgroupSize[0], &WorkgroupSize[1], &WorkgroupSize[2]);
 #endif
 				check(Found == 3);
 				for (uint32 i = 0; i < Found; i++)
@@ -986,7 +986,7 @@ bool FMetalShaderOutputCooker::Build(TArray<uint8>& OutData)
 							check(Var->semantic);
 							unsigned i = (unsigned)strlen(Var->semantic);
 							check(i);
-							while (isdigit(Var->semantic[i-1]))
+							while (isdigit((unsigned char)(Var->semantic[i-1])))
 							{
 								i--;
 							}
