Fix PVS analysis warning.

diff --git a/Engine/Source/Developer/Apple/MetalShaderFormat/Private/MetalDerivedData.cpp b/Engine/Source/Developer/Apple/MetalShaderFormat/Private/MetalDerivedData.cpp
--- a/Engine/Source/Developer/Apple/MetalShaderFormat/Private/MetalDerivedData.cpp
+++ b/Engine/Source/Developer/Apple/MetalShaderFormat/Private/MetalDerivedData.cpp
@@ -752,7 +752,7 @@ bool FMetalShaderOutputCooker::Build(TArray<uint8>& OutData)
 					
 					// UAVs always claim all slots so we don't have conflicts as D3D expects 0-7
 					BufferIndices &= ~(1 << Index);
-					TextureIndices &= ~(1 << Index);
+					TextureIndices &= ~(1llu << uint64(Index));
 					UAVIndices &= ~(1 << Index);
 					
 					OutputData.TypedUAVs |= (1 << Index);
@@ -771,7 +771,7 @@ bool FMetalShaderOutputCooker::Build(TArray<uint8>& OutData)
 					
 					// UAVs always claim all slots so we don't have conflicts as D3D expects 0-7
 					BufferIndices &= ~(1 << Index);
-					TextureIndices &= ~(1 << Index);
+					TextureIndices &= ~(1llu << uint64(Index));
 					UAVIndices &= ~(1 << Index);
 					
 					OutputData.InvariantBuffers |= (1 << Index);
@@ -790,7 +790,7 @@ bool FMetalShaderOutputCooker::Build(TArray<uint8>& OutData)
 					// UAVs always claim all slots so we don't have conflicts as D3D expects 0-7
 					// For texture2d this allows us to emulate atomics with buffers
 					BufferIndices &= ~(1 << Index);
-					TextureIndices &= ~(1 << Index);
+					TextureIndices &= ~(1llu << uint64(Index));
 					UAVIndices &= ~(1 << Index);
 					
 					UAVString += FString::Printf(TEXT("%s%s(%u:%u)"), UAVString.Len() ? TEXT(",") : TEXT(""), UTF8_TO_TCHAR(Binding->name), Index, 1);
@@ -806,7 +806,7 @@ bool FMetalShaderOutputCooker::Build(TArray<uint8>& OutData)
 					
 					// No support for 3-component types in dxc/SPIRV/MSL - need to expose my workarounds there too
 					BufferIndices &= ~(1 << Index);
-					TextureIndices &= ~(1 << Index);
+					TextureIndices &= ~(1llu << uint64(Index));
 					
 					OutputData.TypedBuffers |= (1 << Index);
 					
@@ -899,7 +899,7 @@ bool FMetalShaderOutputCooker::Build(TArray<uint8>& OutData)
 				{
 					check(TextureIndices);
 					uint32 Index = FPlatformMath::CountTrailingZeros64(TextureIndices);
-					TextureIndices &= ~(1 << Index);
+					TextureIndices &= ~(1llu << uint64(Index));
 					
 					SRVString += FString::Printf(TEXT("%s%s(%u:%u)"), SRVString.Len() ? TEXT(",") : TEXT(""), UTF8_TO_TCHAR(Binding->name), Index, 1);
 					
@@ -1174,7 +1174,7 @@ bool FMetalShaderOutputCooker::Build(TArray<uint8>& OutData)
 		}
 		
 		Result = (Results.hasError) ? 0 : 1;
-		if (!Results.hasError)
+		if (!Results.hasError && Results.target)
 		{
 			MetalSource = TCHAR_TO_UTF8(*MetaData);
 			MetalSource += std::string((const char*)Results.target->Data(), Results.target->Size());
