Default to ShaderConductor to compile Metal shaders for the time being to hopefully start exposing more flaws as more people give it a go. For now it is restricted to the mobile renderer shader platforms when using Metal 2.0 or later (where it should work) plus METAL_SM5_NOTESS which is SM5 without tessellation.

diff --git a/Engine/Source/Developer/Apple/MetalShaderFormat/Private/MetalDerivedData.cpp b/Engine/Source/Developer/Apple/MetalShaderFormat/Private/MetalDerivedData.cpp
--- a/Engine/Source/Developer/Apple/MetalShaderFormat/Private/MetalDerivedData.cpp
+++ b/Engine/Source/Developer/Apple/MetalShaderFormat/Private/MetalDerivedData.cpp
@@ -427,9 +427,17 @@ bool FMetalShaderOutputCooker::Build(TArray<uint8>& OutData)
 	bool bUsingTessellation = ((UsingTessellationDefine != nullptr && FString("1") == *UsingTessellationDefine && Frequency == HSF_VertexShader) || Frequency == HSF_HullShader || Frequency == HSF_DomainShader);
 	
 	// Its going to take a while to get dxc+SPIRV working...
+	// The SM5 NoTess is a perfect candidate for defaulting to dxc because it is not used by default and won't break anything
+	static FName NAME_SF_METAL_SM5_NOTESS(TEXT("SF_METAL_SM5_NOTESS"));
+	bool bUseSC = (Input.ShaderFormat == NAME_SF_METAL_SM5_NOTESS);
+	// The Mobile renderer should now be correct from Metal 2.0 and up.
+	bUseSC |= (MetalCompilerTarget == HCT_FeatureLevelES3_1 && VersionEnum >= 3);
+	// Otherwise we must deliberately select the new compiler by setting the Metal version to 6 - which is another cheat (5 is currently Metal 2.1 + IABs)
+	bUseSC |= (VersionEnum >= 6);
+
 #if 1
 #if PLATFORM_MAC || PLATFORM_WINDOWS
-	if (VersionEnum < 6)
+	if (!bUseSC)
 #endif
 #endif
 	{
