Change the way the code is configured to make it easier and more robust to move between hlslcc & ShaderConductor as the shader compiler backend.

diff --git a/Engine/Shaders/Private/InstancedStereo.ush b/Engine/Shaders/Private/InstancedStereo.ush
--- a/Engine/Shaders/Private/InstancedStereo.ush
+++ b/Engine/Shaders/Private/InstancedStereo.ush
@@ -27,7 +27,7 @@ static const float4 EyeClipEdge[2] = { float4(-1.0, 0.0, 0.0, 1.0), float4(1.0,
 #endif
 
 // Generated for Metal in GeneratedInstancedStereo.usf
-#if !(COMPILER_METAL && MAX_SHADER_LANGUAGE_VERSION < 6) && (INSTANCED_STEREO || MOBILE_MULTI_VIEW)
+#if !(COMPILER_METAL && (COMPILER_HLSLCC == 1)) && (INSTANCED_STEREO || MOBILE_MULTI_VIEW)
 ViewState ResolveView(uint ViewIndex)
 {
 	if (ViewIndex == 0)
diff --git a/Engine/Source/Developer/Apple/MetalShaderFormat/Private/MetalShaderCompiler.cpp b/Engine/Source/Developer/Apple/MetalShaderFormat/Private/MetalShaderCompiler.cpp
--- a/Engine/Source/Developer/Apple/MetalShaderFormat/Private/MetalShaderCompiler.cpp
+++ b/Engine/Source/Developer/Apple/MetalShaderFormat/Private/MetalShaderCompiler.cpp
@@ -1957,7 +1957,6 @@ void CompileShader_Metal(const FShaderCompilerInput& _Input,FShaderCompilerOutpu
 		AdditionalDefines.SetDefine(TEXT("MAC"), 1);
 	}
 	
-	AdditionalDefines.SetDefine(TEXT("COMPILER_HLSLCC"), 1);
 	AdditionalDefines.SetDefine(TEXT("COMPILER_METAL"), 1);
 
 	static FName NAME_SF_METAL(TEXT("SF_METAL"));
@@ -2051,8 +2050,20 @@ void CompileShader_Metal(const FShaderCompilerInput& _Input,FShaderCompilerOutpu
 		return;
 	}
 	
-	if (VersionEnum < 6)
+
+	bool bUseSC = (Input.ShaderFormat == NAME_SF_METAL_SM5_NOTESS);
+	// The Mobile renderer should now be correct from Metal 2.0 and up.
+	bUseSC |= (MetalCompilerTarget == HCT_FeatureLevelES3_1 && VersionEnum >= 3);
+	// Otherwise we must deliberately select the new compiler by setting the Metal version to 6 - which is another cheat (5 is currently Metal 2.1 + IABs)
+	bUseSC |= (VersionEnum >= 6);
+
+	if (bUseSC)
+	{
+        AdditionalDefines.SetDefine(TEXT("COMPILER_HLSLCC"), 2);
+	}
+	else
 	{
+        AdditionalDefines.SetDefine(TEXT("COMPILER_HLSLCC"), 1);
 		AdditionalDefines.SetDefine(TEXT("row_major"), TEXT(""));
 	}
 	
diff --git a/Engine/Source/Runtime/Engine/Private/ShaderCompiler/ShaderCompiler.cpp b/Engine/Source/Runtime/Engine/Private/ShaderCompiler/ShaderCompiler.cpp
--- a/Engine/Source/Runtime/Engine/Private/ShaderCompiler/ShaderCompiler.cpp
+++ b/Engine/Source/Runtime/Engine/Private/ShaderCompiler/ShaderCompiler.cpp
@@ -2742,7 +2742,7 @@ ENGINE_API void GenerateInstancedStereoCode(FString& Result)
 	Result += "}\r\n";
 	
 	// ResolveView definition for metal, this allows us to change the branch to a conditional move in the cross compiler
-	Result += "#if COMPILER_METAL && MAX_SHADER_LANGUAGE_VERSION < 6\r\n";
+	Result += "#if COMPILER_METAL && (COMPILER_HLSLCC == 1)\r\n";
 	Result += "ViewState ResolveView(uint ViewIndex)\r\n";
 	Result += "{\r\n";
 	Result += "\tViewState Result;\r\n";
