Turn on/off the use of ShaderConductor using Yuri's new r.D3D.ForceDXC flag on Metal as well - it makes it easier/safer to toggle and is more consistent.

diff --git a/Engine/Source/Developer/Apple/MetalShaderFormat/Private/MetalDerivedData.cpp b/Engine/Source/Developer/Apple/MetalShaderFormat/Private/MetalDerivedData.cpp
--- a/Engine/Source/Developer/Apple/MetalShaderFormat/Private/MetalDerivedData.cpp
+++ b/Engine/Source/Developer/Apple/MetalShaderFormat/Private/MetalDerivedData.cpp
@@ -427,13 +427,7 @@ bool FMetalShaderOutputCooker::Build(TArray<uint8>& OutData)
 	bool bUsingTessellation = ((UsingTessellationDefine != nullptr && FString("1") == *UsingTessellationDefine && Frequency == HSF_VertexShader) || Frequency == HSF_HullShader || Frequency == HSF_DomainShader);
 	
 	// Its going to take a while to get dxc+SPIRV working...
-	// The SM5 NoTess is a perfect candidate for defaulting to dxc because it is not used by default and won't break anything
-	static FName NAME_SF_METAL_SM5_NOTESS(TEXT("SF_METAL_SM5_NOTESS"));
-	bool bUseSC = (Input.ShaderFormat == NAME_SF_METAL_SM5_NOTESS);
-	// The Mobile renderer should now be correct from Metal 2.0 and up.
-	bUseSC |= (MetalCompilerTarget == HCT_FeatureLevelES3_1 && VersionEnum >= 3);
-	// Otherwise we must deliberately select the new compiler by setting the Metal version to 6 - which is another cheat (5 is currently Metal 2.1 + IABs)
-	bUseSC |= (VersionEnum >= 6);
+	bool const bUseSC = Input.Environment.CompilerFlags.Contains(CFLAG_ForceDXC);
 
 #if 1
 #if PLATFORM_MAC || PLATFORM_WINDOWS
diff --git a/Engine/Source/Developer/Apple/MetalShaderFormat/Private/MetalShaderCompiler.cpp b/Engine/Source/Developer/Apple/MetalShaderFormat/Private/MetalShaderCompiler.cpp
--- a/Engine/Source/Developer/Apple/MetalShaderFormat/Private/MetalShaderCompiler.cpp
+++ b/Engine/Source/Developer/Apple/MetalShaderFormat/Private/MetalShaderCompiler.cpp
@@ -2051,12 +2051,7 @@ void CompileShader_Metal(const FShaderCompilerInput& _Input,FShaderCompilerOutpu
 	}
 	
 
-	bool bUseSC = (Input.ShaderFormat == NAME_SF_METAL_SM5_NOTESS);
-	// The Mobile renderer should now be correct from Metal 2.0 and up.
-	bUseSC |= (MetalCompilerTarget == HCT_FeatureLevelES3_1 && VersionEnum >= 3);
-	// Otherwise we must deliberately select the new compiler by setting the Metal version to 6 - which is another cheat (5 is currently Metal 2.1 + IABs)
-	bUseSC |= (VersionEnum >= 6);
-
+	bool const bUseSC = Input.Environment.CompilerFlags.Contains(CFLAG_ForceDXC);
 	if (bUseSC)
 	{
         AdditionalDefines.SetDefine(TEXT("COMPILER_HLSLCC"), 2);
diff --git a/Engine/Source/Runtime/Engine/Private/ShaderCompiler/ShaderCompiler.cpp b/Engine/Source/Runtime/Engine/Private/ShaderCompiler/ShaderCompiler.cpp
--- a/Engine/Source/Runtime/Engine/Private/ShaderCompiler/ShaderCompiler.cpp
+++ b/Engine/Source/Runtime/Engine/Private/ShaderCompiler/ShaderCompiler.cpp
@@ -3222,6 +3222,13 @@ void GlobalBeginCompileShader(
 				Input.Environment.CompilerFlags.Add(CFLAG_NoFastMath);
 			}
 		}
+		{
+			static const auto CVar = IConsoleManager::Get().FindConsoleVariable(TEXT("r.D3D.ForceDXC"));
+			if (CVar && CVar->GetInt() != 0)
+			{
+				Input.Environment.CompilerFlags.Add(CFLAG_ForceDXC);
+			}
+		}
 	}
 
 	if (IsOpenGLPlatform((EShaderPlatform)Target.Platform) && 
diff --git a/Engine/Source/Runtime/RenderCore/Private/Shader.cpp b/Engine/Source/Runtime/RenderCore/Private/Shader.cpp
--- a/Engine/Source/Runtime/RenderCore/Private/Shader.cpp
+++ b/Engine/Source/Runtime/RenderCore/Private/Shader.cpp
@@ -2315,6 +2315,13 @@ void ShaderMapAppendKeyString(EShaderPlatform Platform, FString& KeyString)
 		{
 			KeyString += TEXT("_ARCHIVE");
 		}
+		{
+			static const auto CVar = IConsoleManager::Get().FindConsoleVariable(TEXT("r.D3D.ForceDXC"));
+			if (CVar && CVar->GetInt() != 0)
+			{
+				KeyString += TEXT("_DXC");
+			}
+		}
 	}
 
 	if (IsFeatureLevelSupported(Platform, ERHIFeatureLevel::SM4))
