DR - Fix for ShaderConductor dlls
#rnx

diff --git a/Engine/Source/Developer/Apple/MetalShaderFormat/MetalShaderFormat.Build.cs b/Engine/Source/Developer/Apple/MetalShaderFormat/MetalShaderFormat.Build.cs
--- a/Engine/Source/Developer/Apple/MetalShaderFormat/MetalShaderFormat.Build.cs
+++ b/Engine/Source/Developer/Apple/MetalShaderFormat/MetalShaderFormat.Build.cs
@@ -43,5 +43,16 @@ public class MetalShaderFormat : ModuleRules
 				"DerivedDataCache",
 			}
 			);
+
+		if (Target.Platform == UnrealTargetPlatform.Win64)
+		{
+			string DllsPath = "$(EngineDir)/Binaries/ThirdParty/ShaderConductor/Win64/";
+
+			RuntimeDependencies.Add("$(TargetOutputDir)/dxcompiler_sc.dll", DllsPath + "dxcompiler_sc.dll");
+			RuntimeDependencies.Add("$(TargetOutputDir)/ShaderConductor.dll", DllsPath + "ShaderConductor.dll");
+
+			PublicDelayLoadDLLs.Add("dxcompiler_sc.dll");
+			PublicDelayLoadDLLs.Add("ShaderConductor.dll");
+		}
 	}
 }
