Fix DXC compiled Metal shaders not exporting the proper list of buffer bindings.

diff --git a/Engine/Source/Developer/Apple/MetalShaderFormat/Private/MetalShaderCompiler.cpp b/Engine/Source/Developer/Apple/MetalShaderFormat/Private/MetalShaderCompiler.cpp
--- a/Engine/Source/Developer/Apple/MetalShaderFormat/Private/MetalShaderCompiler.cpp
+++ b/Engine/Source/Developer/Apple/MetalShaderFormat/Private/MetalShaderCompiler.cpp
@@ -1102,7 +1102,7 @@ void BuildMetalShaderOutput(
 			}
 		}
 		
-		if (Version == 6)
+		if (Version == 6 || ShaderInput.Environment.CompilerFlags.Contains(CFLAG_ForceDXC))
 		{
 			Header.Bindings.LinearBuffer = Header.Bindings.TypedBuffers;
 			Header.Bindings.TypedBuffers = 0;
