Surface errors from the Metal shader compilation tools.

diff --git a/Engine/Source/Developer/Apple/MetalShaderFormat/Private/MetalDerivedData.cpp b/Engine/Source/Developer/Apple/MetalShaderFormat/Private/MetalDerivedData.cpp
--- a/Engine/Source/Developer/Apple/MetalShaderFormat/Private/MetalDerivedData.cpp
+++ b/Engine/Source/Developer/Apple/MetalShaderFormat/Private/MetalDerivedData.cpp
@@ -263,20 +263,23 @@ bool FMetalShaderBytecodeCooker::Build(TArray<uint8>& OutData)
 			}
 		}
 
-		RemoveRemoteFile(RemoteObjFile);
-		CopyRemoteFileToLocal(RemoteOutputFilename, Job.OutputFile);
-		Output.NativePath = RemoteInputFile;
-		bSucceeded = FFileHelper::LoadFileToArray(Output.OutputFile, *Job.OutputFile);
+		if (bSucceeded)
+		{
+			RemoveRemoteFile(RemoteObjFile);
+			CopyRemoteFileToLocal(RemoteOutputFilename, Job.OutputFile);
+			Output.NativePath = RemoteInputFile;
+			bSucceeded = FFileHelper::LoadFileToArray(Output.OutputFile, *Job.OutputFile);
 
+			if (!bSucceeded)
+			{
+				Job.Message = FString::Printf(TEXT("Failed to load output file: %s"), *Job.OutputFile);
+			}
+		}
+		
 		if (!Job.bCompileAsPCH)
 		{
 			RemoveRemoteFile(RemoteOutputFilename);
 		}
-
-		if (!bSucceeded)
-		{
-			Job.Message = FString::Printf(TEXT("Failed to load output file: %s"), *Job.OutputFile);
-		}
 	}
 	else
 	{
