Keep one of the changes that appears to have been removed.

diff --git a/Engine/Source/Developer/Apple/MetalShaderFormat/Private/MetalDerivedData.cpp b/Engine/Source/Developer/Apple/MetalShaderFormat/Private/MetalDerivedData.cpp
--- a/Engine/Source/Developer/Apple/MetalShaderFormat/Private/MetalDerivedData.cpp
+++ b/Engine/Source/Developer/Apple/MetalShaderFormat/Private/MetalDerivedData.cpp
@@ -255,6 +255,8 @@ bool FMetalShaderBytecodeCooker::Build(TArray<uint8>& OutData)
 					{
 						Job.Message = FString::Printf(TEXT("Failed to load object file: %s"), *Job.OutputObjectFile);
 					}
+
+					RemoveRemoteFile(RemoteObjFile);
 				}
 			}
 			else
