Don't export unused UAV counter buffers when compiling Metal shaders with DXC & use the new texture-buffer flag to enable the newer type when available.

diff --git a/Engine/Source/Developer/Apple/MetalShaderFormat/Private/MetalDerivedData.cpp b/Engine/Source/Developer/Apple/MetalShaderFormat/Private/MetalDerivedData.cpp
--- a/Engine/Source/Developer/Apple/MetalShaderFormat/Private/MetalDerivedData.cpp
+++ b/Engine/Source/Developer/Apple/MetalShaderFormat/Private/MetalDerivedData.cpp
@@ -603,6 +603,7 @@ bool FMetalShaderOutputCooker::Build(TArray<uint8>& OutData)
 			SpvReflectResult SPVRResult = SPV_REFLECT_RESULT_NOT_READY;
 			uint32 Count = 0;
 			TArray<SpvReflectDescriptorBinding*> Bindings;
+			TSet<SpvReflectDescriptorBinding*> Counters;
 			TArray<SpvReflectInterfaceVariable*> InputVars;
 			TArray<SpvReflectInterfaceVariable*> OutputVars;
 			
@@ -719,6 +720,11 @@ bool FMetalShaderOutputCooker::Build(TArray<uint8>& OutData)
 						}
 						case SPV_REFLECT_RESOURCE_FLAG_UAV:
 						{
+							if (Binding->uav_counter_binding)
+							{
+								Counters.Add(Binding->uav_counter_binding);
+							}
+							
 							switch(Binding->descriptor_type)
 							{
 								case SPV_REFLECT_DESCRIPTOR_TYPE_STORAGE_IMAGE:
@@ -733,7 +739,10 @@ bool FMetalShaderOutputCooker::Build(TArray<uint8>& OutData)
 								}
 								case SPV_REFLECT_DESCRIPTOR_TYPE_STORAGE_BUFFER:
 								{
-									SBufferUAVBindings.Add(Binding);
+									if (!Counters.Contains(Binding) || Binding->accessed)
+									{
+										SBufferUAVBindings.Add(Binding);
+									}
 									break;
 								}
 								default:
@@ -1119,12 +1128,35 @@ bool FMetalShaderOutputCooker::Build(TArray<uint8>& OutData)
 			
 			TargetDesc.language = ShaderConductor::ShadingLanguage::Msl;
 			
+			ShaderConductor::MacroDefine Defines[4] = {{"texel_buffer_texture_width", "0"}, {nullptr, nullptr}, {nullptr, nullptr}, {nullptr, nullptr}};
+			TargetDesc.numOptions = 1;
+			TargetDesc.options = &Defines[0];
+			switch(Semantics)
+			{
+				case EMetalGPUSemanticsImmediateDesktop:
+					TargetDesc.platform = "macOS";
+					break;
+				case EMetalGPUSemanticsTBDRDesktop:
+					TargetDesc.platform = "iOS";
+					Defines[1] = { "ios_support_base_vertex_instance", "1" };
+					Defines[2] = { "ios_use_framebuffer_fetch_subpasses", "1" };
+					TargetDesc.numOptions += 2;
+					break;
+				case EMetalGPUSemanticsMobile:
+				default:
+					TargetDesc.platform = "iOS";
+					Defines[1] = { "ios_use_framebuffer_fetch_subpasses", "1" };
+					TargetDesc.numOptions += 1;
+					break;
+			}
+			
 			switch (VersionEnum)
 			{
 				case 6:
 				case 5:
 				case 4:
 				{
+					Defines[TargetDesc.numOptions++] = { "texture_buffer_native", "1" };
 					TargetDesc.version = "20100";
 					break;
 				}
@@ -1151,28 +1183,6 @@ bool FMetalShaderOutputCooker::Build(TArray<uint8>& OutData)
 				}
 			}
 			
-			ShaderConductor::MacroDefine Defines[3] = {{"texel_buffer_texture_width", "0"}, {nullptr, nullptr}, {nullptr, nullptr}};
-			TargetDesc.numOptions = 1;
-			TargetDesc.options = &Defines[0];
-			switch(Semantics)
-			{
-				case EMetalGPUSemanticsImmediateDesktop:
-					TargetDesc.platform = "macOS";
-					break;
-				case EMetalGPUSemanticsTBDRDesktop:
-					TargetDesc.platform = "iOS";
-					Defines[1] = { "ios_support_base_vertex_instance", "1" };
-					Defines[2] = { "ios_use_framebuffer_fetch_subpasses", "1" };
-					TargetDesc.numOptions += 2;
-					break;
-				case EMetalGPUSemanticsMobile:
-				default:
-					TargetDesc.platform = "iOS";
-					Defines[1] = { "ios_use_framebuffer_fetch_subpasses", "1" };
-					TargetDesc.numOptions += 1;
-					break;
-			}
-			
 			ShaderConductor::Blob* IRBlob = Results.target;
 			
 			Results = ShaderConductor::Compiler::ConvertBinary(Results, SourceDesc, TargetDesc);
