Export the buffer side-table binding when compiling shaders for Metal with DXC.

diff --git a/Engine/Source/Developer/Apple/MetalShaderFormat/Private/MetalDerivedData.cpp b/Engine/Source/Developer/Apple/MetalShaderFormat/Private/MetalDerivedData.cpp
--- a/Engine/Source/Developer/Apple/MetalShaderFormat/Private/MetalDerivedData.cpp
+++ b/Engine/Source/Developer/Apple/MetalShaderFormat/Private/MetalDerivedData.cpp
@@ -593,6 +593,7 @@ bool FMetalShaderOutputCooker::Build(TArray<uint8>& OutData)
 		TargetDesc.language = ShaderConductor::ShadingLanguage::SpirV;
 		ShaderConductor::Compiler::ResultDesc Results = ShaderConductor::Compiler::Compile(SourceDesc, Options, TargetDesc);
 		Result = (Results.hasError) ? 0 : 1;
+		uint32 BufferIndices = 0xffffffff;
 		if (!Results.hasError)
 		{
 			// Now perform reflection on the SPIRV and tweak any decorations that we need to.
@@ -608,7 +609,6 @@ bool FMetalShaderOutputCooker::Build(TArray<uint8>& OutData)
 			TArray<SpvReflectInterfaceVariable*> OutputVars;
 			
 			uint8 UAVIndices = 0xff;
-			uint32 BufferIndices = 0xffffffff;
 			uint64 TextureIndices = 0xffffffffffffffff;
 			uint16 SamplerIndices = 0xffff;
 			
@@ -1105,7 +1105,6 @@ bool FMetalShaderOutputCooker::Build(TArray<uint8>& OutData)
 			{
 				MetaData += FString::Printf(TEXT("// @NumThreads: %s\n"), *WKGString);
 			}
-			MetaData += FString::Printf(TEXT("\n\n"), *WKGString);
 			
 			ShaderConductor::Blob* OldData = Results.target;
 			Results.target = ShaderConductor::CreateBlob(Reflection.GetCode(), Reflection.GetCodeSize());
@@ -1121,6 +1120,7 @@ bool FMetalShaderOutputCooker::Build(TArray<uint8>& OutData)
 			ShaderConductor::DestroyBlob(OldData);
 		}
 		
+		uint32 SideTableIndex = 0;
 		if (Result)
 		{
 			ShaderConductor::DestroyBlob(Results.errorWarningMsg);
@@ -1128,8 +1128,12 @@ bool FMetalShaderOutputCooker::Build(TArray<uint8>& OutData)
 			
 			TargetDesc.language = ShaderConductor::ShadingLanguage::Msl;
 			
-			ShaderConductor::MacroDefine Defines[4] = {{"texel_buffer_texture_width", "0"}, {nullptr, nullptr}, {nullptr, nullptr}, {nullptr, nullptr}};
-			TargetDesc.numOptions = 1;
+			SideTableIndex = FPlatformMath::CountTrailingZeros(BufferIndices);
+			char BufferIdx[3];
+			FCStringAnsi::Snprintf(BufferIdx, 3, "%d", SideTableIndex);
+			
+			ShaderConductor::MacroDefine Defines[6] = {{"texel_buffer_texture_width", "0"}, {"enforce_storge_buffer_bounds", "1"}, {"metadata_buffer_index", BufferIdx}, {nullptr, nullptr}, {nullptr, nullptr}, {nullptr, nullptr}};
+			TargetDesc.numOptions = 3;
 			TargetDesc.options = &Defines[0];
 			switch(Semantics)
 			{
@@ -1138,14 +1142,13 @@ bool FMetalShaderOutputCooker::Build(TArray<uint8>& OutData)
 					break;
 				case EMetalGPUSemanticsTBDRDesktop:
 					TargetDesc.platform = "iOS";
-					Defines[1] = { "ios_support_base_vertex_instance", "1" };
-					Defines[2] = { "ios_use_framebuffer_fetch_subpasses", "1" };
-					TargetDesc.numOptions += 2;
+					Defines[TargetDesc.numOptions++] = { "ios_support_base_vertex_instance", "1" };
+					Defines[TargetDesc.numOptions++] = { "ios_use_framebuffer_fetch_subpasses", "1" };
 					break;
 				case EMetalGPUSemanticsMobile:
 				default:
 					TargetDesc.platform = "iOS";
-					Defines[1] = { "ios_use_framebuffer_fetch_subpasses", "1" };
+					Defines[TargetDesc.numOptions++] = { "ios_use_framebuffer_fetch_subpasses", "1" };
 					TargetDesc.numOptions += 1;
 					break;
 			}
@@ -1193,6 +1196,12 @@ bool FMetalShaderOutputCooker::Build(TArray<uint8>& OutData)
 		Result = (Results.hasError) ? 0 : 1;
 		if (!Results.hasError && Results.target)
 		{
+			if(strstr((const char*)Results.target->Data(), "spvBufferSizes"))
+			{
+				MetaData += FString::Printf(TEXT("// @SideTable: spvBufferSizes(%d)\n"), SideTableIndex);
+			}
+			MetaData += FString::Printf(TEXT("\n\n"));
+			
 			MetalSource = TCHAR_TO_UTF8(*MetaData);
 			MetalSource += std::string((const char*)Results.target->Data(), Results.target->Size());
 			
diff --git a/Engine/Source/Runtime/Apple/MetalRHI/Private/MetalPipeline.cpp b/Engine/Source/Runtime/Apple/MetalRHI/Private/MetalPipeline.cpp
--- a/Engine/Source/Runtime/Apple/MetalRHI/Private/MetalPipeline.cpp
+++ b/Engine/Source/Runtime/Apple/MetalRHI/Private/MetalPipeline.cpp
@@ -543,7 +543,7 @@ APPLE_PLATFORM_OBJECT_ALLOC_OVERRIDES(FMetalShaderPipeline)
 			case MTLArgumentTypeBuffer:
 			{
 				checkf(Arg.index < ML_MaxBuffers, TEXT("Metal buffer index exceeded!"));
-				if (FString(Arg.name) != TEXT("BufferSizes"))
+				if (FString(Arg.name) != TEXT("BufferSizes") && FString(Arg.name) != TEXT("spvBufferSizes"))
 				{
 					ResourceMask[Frequency].BufferMask |= (1 << Arg.index);
 				
