Make DXC errors report without causing an assert in the virtual file-path handling.

diff --git a/Engine/Source/Developer/Apple/MetalShaderFormat/Private/MetalDerivedData.cpp b/Engine/Source/Developer/Apple/MetalShaderFormat/Private/MetalDerivedData.cpp
--- a/Engine/Source/Developer/Apple/MetalShaderFormat/Private/MetalDerivedData.cpp
+++ b/Engine/Source/Developer/Apple/MetalShaderFormat/Private/MetalDerivedData.cpp
@@ -1223,7 +1223,8 @@ bool FMetalShaderOutputCooker::Build(TArray<uint8>& OutData)
 		
 		if (Results.errorWarningMsg)
 		{
-			MetalErrors = ANSI_TO_TCHAR((char const*)Results.errorWarningMsg->Data());
+			std::string ErrorText((const char*)Results.errorWarningMsg->Data(), Results.errorWarningMsg->Size());
+			MetalErrors = Input.VirtualSourceFilePath + TEXT("(0): ") + ANSI_TO_TCHAR((char const*)ErrorText.c_str());
 			ShaderConductor::DestroyBlob(Results.errorWarningMsg);
 			Results.errorWarningMsg = nullptr;
 		}
