More changes to ensure that DXC reports errors in a way that our code can handle.

diff --git a/Engine/Source/Developer/Apple/MetalShaderFormat/Private/MetalDerivedData.cpp b/Engine/Source/Developer/Apple/MetalShaderFormat/Private/MetalDerivedData.cpp
--- a/Engine/Source/Developer/Apple/MetalShaderFormat/Private/MetalDerivedData.cpp
+++ b/Engine/Source/Developer/Apple/MetalShaderFormat/Private/MetalDerivedData.cpp
@@ -1224,7 +1224,15 @@ bool FMetalShaderOutputCooker::Build(TArray<uint8>& OutData)
 		if (Results.errorWarningMsg)
 		{
 			std::string ErrorText((const char*)Results.errorWarningMsg->Data(), Results.errorWarningMsg->Size());
-			MetalErrors = Input.VirtualSourceFilePath + TEXT("(0): ") + ANSI_TO_TCHAR((char const*)ErrorText.c_str());
+			
+			TArray<FString> ErrorLines;
+			MetalErrors = ANSI_TO_TCHAR((char const*)ErrorText.c_str());
+			MetalErrors.ParseIntoArray(ErrorLines, TEXT("\n"), true);
+			MetalErrors = TEXT("");
+			for (int32 LineIndex = 0; LineIndex < ErrorLines.Num(); ++LineIndex)
+			{
+				MetalErrors += Input.VirtualSourceFilePath + TEXT("(0): ") + ErrorLines[LineIndex] + TEXT("\n");
+			}
 			ShaderConductor::DestroyBlob(Results.errorWarningMsg);
 			Results.errorWarningMsg = nullptr;
 		}
