Metal vertex shaders in tessellation pipelines must capture their output into a buffer so that it can feed into the compute-based hull shader. This is different than the HLSLCC approach which combined the Vertex & Hull shader, which we can't easily do with DXC.

diff --git a/Engine/Source/Developer/Apple/MetalShaderFormat/Private/MetalDerivedData.cpp b/Engine/Source/Developer/Apple/MetalShaderFormat/Private/MetalDerivedData.cpp
--- a/Engine/Source/Developer/Apple/MetalShaderFormat/Private/MetalDerivedData.cpp
+++ b/Engine/Source/Developer/Apple/MetalShaderFormat/Private/MetalDerivedData.cpp
@@ -1131,8 +1131,8 @@ bool FMetalShaderOutputCooker::Build(TArray<uint8>& OutData)
 			SideTableIndex = FPlatformMath::CountTrailingZeros(BufferIndices);
 			char BufferIdx[3];
 			FCStringAnsi::Snprintf(BufferIdx, 3, "%d", SideTableIndex);
-			
-			ShaderConductor::MacroDefine Defines[6] = {{"texel_buffer_texture_width", "0"}, {"enforce_storge_buffer_bounds", "1"}, {"metadata_buffer_index", BufferIdx}, {nullptr, nullptr}, {nullptr, nullptr}, {nullptr, nullptr}};
+
+			ShaderConductor::MacroDefine Defines[7] = {{"texel_buffer_texture_width", "0"}, {"enforce_storge_buffer_bounds", "1"}, {"metadata_buffer_index", BufferIdx}, {nullptr, nullptr}, {nullptr, nullptr}, {nullptr, nullptr}, {nullptr, nullptr}};
 			TargetDesc.numOptions = 3;
 			TargetDesc.options = &Defines[0];
 			switch(Semantics)
@@ -1153,6 +1153,11 @@ bool FMetalShaderOutputCooker::Build(TArray<uint8>& OutData)
 					break;
 			}
 			
+			if (Frequency == HSF_VertexShader && bUsingTessellation)
+			{
+				Defines[TargetDesc.numOptions++] = { "capture_output_to_buffer", "1" };
+			}
+			
 			switch (VersionEnum)
 			{
 				case 6:
