Map the obvious SPIRV execution mode variables to the Metal metadata output.

diff --git a/Engine/Source/Developer/Apple/MetalShaderFormat/Private/MetalDerivedData.cpp b/Engine/Source/Developer/Apple/MetalShaderFormat/Private/MetalDerivedData.cpp
--- a/Engine/Source/Developer/Apple/MetalShaderFormat/Private/MetalDerivedData.cpp
+++ b/Engine/Source/Developer/Apple/MetalShaderFormat/Private/MetalDerivedData.cpp
@@ -607,6 +607,7 @@ bool FMetalShaderOutputCooker::Build(TArray<uint8>& OutData)
 			TSet<SpvReflectDescriptorBinding*> Counters;
 			TArray<SpvReflectInterfaceVariable*> InputVars;
 			TArray<SpvReflectInterfaceVariable*> OutputVars;
+			TArray<SpvReflectExecutionMode*> ExecutionModes;
 			
 			uint8 UAVIndices = 0xff;
 			uint64 TextureIndices = 0xffffffffffffffff;
@@ -621,23 +622,119 @@ bool FMetalShaderOutputCooker::Build(TArray<uint8>& OutData)
 			FString INPString;
 			FString OUTString;
 			FString WKGString;
+			FString TESString;
 			
-			if (Frequency == HSF_ComputeShader)
 			{
-				size_t ThreadSizeIdx = SourceData.find("numthreads");
-				check(ThreadSizeIdx != std::string::npos);
-				char const* String = SourceData.c_str() + ThreadSizeIdx;
-				uint32 WorkgroupSize[3];
-
-#if !PLATFORM_WINDOWS
-				size_t Found = sscanf(String, "numthreads(%u, %u, %u)", &WorkgroupSize[0], &WorkgroupSize[1], &WorkgroupSize[2]);
-#else
-				size_t Found = sscanf_s(String, "numthreads(%u, %u, %u)", &WorkgroupSize[0], &WorkgroupSize[1], &WorkgroupSize[2]);
-#endif
-				check(Found == 3);
-				for (uint32 i = 0; i < Found; i++)
+				Count = 0;
+				SPVRResult = Reflection.EnumerateExecutionModes(&Count, nullptr);
+				check(SPVRResult == SPV_REFLECT_RESULT_SUCCESS);
+				ExecutionModes.SetNum(Count);
+				SPVRResult = Reflection.EnumerateExecutionModes(&Count, ExecutionModes.GetData());
+				check(SPVRResult == SPV_REFLECT_RESULT_SUCCESS);
+				for (uint32 i = 0; i < Count; i++)
 				{
-					WKGString += FString::Printf(TEXT("%s%u"), WKGString.Len() ? TEXT(", ") : TEXT(""), WorkgroupSize[i]);
+					auto* Mode = ExecutionModes[i];
+					switch (Mode->mode) {
+						case SpvExecutionModeLocalSize:
+						case SpvExecutionModeLocalSizeHint:
+							if (Frequency == HSF_ComputeShader)
+							{
+								check(Mode->operands_count == 3);
+								for (uint32 j = 0; j < Mode->operands_count; j++)
+								{
+									WKGString += FString::Printf(TEXT("%s%u"), WKGString.Len() ? TEXT(", ") : TEXT(""), Mode->operands[j]);
+								}
+							}
+							break;
+						case SpvExecutionModeTriangles:
+							if (Frequency == HSF_HullShader || Frequency == HSF_DomainShader)
+							{
+								TESString += TEXT("// @TessellationDomain: tri\n");
+							}
+							break;
+						case SpvExecutionModeQuads:
+							if (Frequency == HSF_HullShader || Frequency == HSF_DomainShader)
+							{
+								TESString += TEXT("// @TessellationDomain: quad\n");
+							}
+							break;
+						case SpvExecutionModeIsolines:
+							if (Frequency == HSF_HullShader || Frequency == HSF_DomainShader)
+							{
+								check(0); // Not supported by Metal
+							}
+							break;
+						case SpvExecutionModeOutputVertices:
+							if (Frequency == HSF_HullShader || Frequency == HSF_DomainShader)
+							{
+								check(Mode->operands_count == 1);
+								TESString += FString::Printf(TEXT("// @TessellationOutputControlPoints: %d\n"), Mode->operands[0]);
+							}
+							break;
+						case SpvExecutionModeVertexOrderCw:
+							if (Frequency == HSF_HullShader || Frequency == HSF_DomainShader)
+							{
+								TESString += TEXT("// @TessellationOutputWinding: cw\n");
+							}
+							break;
+						case SpvExecutionModeVertexOrderCcw:
+							if (Frequency == HSF_HullShader || Frequency == HSF_DomainShader)
+							{
+								TESString += TEXT("// @TessellationOutputWinding: ccw\n");
+							}
+							break;
+						case SpvExecutionModeSpacingEqual:
+							if (Frequency == HSF_HullShader || Frequency == HSF_DomainShader)
+							{
+								TESString += TEXT("// @TessellationPartitioning: integer\n");
+							}
+							break;
+						case SpvExecutionModeSpacingFractionalEven:
+							if (Frequency == HSF_HullShader || Frequency == HSF_DomainShader)
+							{
+								TESString += TEXT("// @TessellationPartitioning: fractional_even\n");
+							}
+							break;
+						case SpvExecutionModeSpacingFractionalOdd:
+							if (Frequency == HSF_HullShader || Frequency == HSF_DomainShader)
+							{
+								TESString += TEXT("// @TessellationPartitioning: fractional_odd\n");
+							}
+							break;
+						case SpvExecutionModeInvocations:
+						case SpvExecutionModePixelCenterInteger:
+						case SpvExecutionModeOriginUpperLeft:
+						case SpvExecutionModeOriginLowerLeft:
+						case SpvExecutionModeEarlyFragmentTests:
+						case SpvExecutionModePointMode:
+						case SpvExecutionModeXfb:
+						case SpvExecutionModeDepthReplacing:
+						case SpvExecutionModeDepthGreater:
+						case SpvExecutionModeDepthLess:
+						case SpvExecutionModeDepthUnchanged:
+						case SpvExecutionModeInputPoints:
+						case SpvExecutionModeInputLines:
+						case SpvExecutionModeInputLinesAdjacency:
+						case SpvExecutionModeInputTrianglesAdjacency:
+						case SpvExecutionModeOutputPoints:
+						case SpvExecutionModeOutputLineStrip:
+						case SpvExecutionModeOutputTriangleStrip:
+						case SpvExecutionModeVecTypeHint:
+						case SpvExecutionModeContractionOff:
+						case SpvExecutionModeInitializer:
+						case SpvExecutionModeFinalizer:
+						case SpvExecutionModeSubgroupSize:
+						case SpvExecutionModeSubgroupsPerWorkgroup:
+						case SpvExecutionModeSubgroupsPerWorkgroupId:
+						case SpvExecutionModeLocalSizeId:
+						case SpvExecutionModeLocalSizeHintId:
+						case SpvExecutionModePostDepthCoverage:
+						case SpvExecutionModeStencilRefReplacingEXT:
+							break;
+							
+						default:
+							break;
+					}
 				}
 			}
 			
@@ -1101,6 +1198,10 @@ bool FMetalShaderOutputCooker::Build(TArray<uint8>& OutData)
 			{
 				MetaData += FString::Printf(TEXT("// @SamplerStates: %s\n"), *SMPString);
 			}
+			if (TESString.Len())
+			{
+				MetaData += TESString;
+			}
 			if (WKGString.Len())
 			{
 				MetaData += FString::Printf(TEXT("// @NumThreads: %s\n"), *WKGString);
