A bit more work hooking up the various buffer indices into the Metal shader metadata - still more to be done before it will actually work and pass through everything we need.

diff --git a/Engine/Source/Developer/Apple/MetalShaderFormat/Private/MetalDerivedData.cpp b/Engine/Source/Developer/Apple/MetalShaderFormat/Private/MetalDerivedData.cpp
--- a/Engine/Source/Developer/Apple/MetalShaderFormat/Private/MetalDerivedData.cpp
+++ b/Engine/Source/Developer/Apple/MetalShaderFormat/Private/MetalDerivedData.cpp
@@ -518,6 +518,8 @@ bool FMetalShaderOutputCooker::Build(TArray<uint8>& OutData)
         SourceDesc.entryPoint = EntryPointName.c_str();
 		SourceDesc.numDefines = 0;
 		SourceDesc.defines = nullptr;
+		
+		FString TESString;
         
         switch (Frequency)
         {
@@ -559,7 +561,7 @@ bool FMetalShaderOutputCooker::Build(TArray<uint8>& OutData)
 		
 		ShaderConductor::Blob* RewriteBlob = nullptr;
 		
-        // Can't rewrite tessellation shaders because that'll prevent combining Vertex and Hull
+        // Can't rewrite tessellation shaders because the rewriter doesn't handle HLSL attributes like patchFunction properly and it isn't clear how it should.
         if (!bUsingTessellation)
         {
             ShaderConductor::Compiler::ResultDesc Results = ShaderConductor::Compiler::Rewrite(SourceDesc, Options);
@@ -622,7 +624,6 @@ bool FMetalShaderOutputCooker::Build(TArray<uint8>& OutData)
 			FString INPString;
 			FString OUTString;
 			FString WKGString;
-			FString TESString;
 			
 			{
 				Count = 0;
@@ -1198,10 +1199,6 @@ bool FMetalShaderOutputCooker::Build(TArray<uint8>& OutData)
 			{
 				MetaData += FString::Printf(TEXT("// @SamplerStates: %s\n"), *SMPString);
 			}
-			if (TESString.Len())
-			{
-				MetaData += TESString;
-			}
 			if (WKGString.Len())
 			{
 				MetaData += FString::Printf(TEXT("// @NumThreads: %s\n"), *WKGString);
@@ -1222,6 +1219,11 @@ bool FMetalShaderOutputCooker::Build(TArray<uint8>& OutData)
 		}
 		
 		uint32 SideTableIndex = 0;
+		uint32 OutputBufferIndex = 0;
+		uint32 IndirectParamsIndex = 0;
+		uint32 PatchBufferIndex = 0;
+		uint32 TessFactorBufferIndex = 0;
+		uint32 ShaderInputWGIndex = 0;
 		if (Result)
 		{
 			ShaderConductor::DestroyBlob(Results.errorWarningMsg);
@@ -1232,8 +1234,9 @@ bool FMetalShaderOutputCooker::Build(TArray<uint8>& OutData)
 			SideTableIndex = FPlatformMath::CountTrailingZeros(BufferIndices);
 			char BufferIdx[3];
 			FCStringAnsi::Snprintf(BufferIdx, 3, "%d", SideTableIndex);
+			BufferIndices &= ~(1 << SideTableIndex);
 
-			ShaderConductor::MacroDefine Defines[7] = {{"texel_buffer_texture_width", "0"}, {"enforce_storge_buffer_bounds", "1"}, {"buffer_size_buffer_index", BufferIdx}, {nullptr, nullptr}, {nullptr, nullptr}, {nullptr, nullptr}, {nullptr, nullptr}};
+			ShaderConductor::MacroDefine Defines[12] = {{"texel_buffer_texture_width", "0"}, {"enforce_storge_buffer_bounds", "1"}, {"buffer_size_buffer_index", BufferIdx}, {nullptr, nullptr}, {nullptr, nullptr}, {nullptr, nullptr}, {nullptr, nullptr}};
 			TargetDesc.numOptions = 3;
 			TargetDesc.options = &Defines[0];
 			switch(Semantics)
@@ -1259,6 +1262,47 @@ bool FMetalShaderOutputCooker::Build(TArray<uint8>& OutData)
 				Defines[TargetDesc.numOptions++] = { "capture_output_to_buffer", "1" };
 			}
 			
+			char OutputBufferIdx[3];
+			char IndirectParamsIdx[3];
+			if ((Frequency == HSF_VertexShader && bUsingTessellation) || Frequency == HSF_HullShader)
+			{
+				OutputBufferIndex = FPlatformMath::CountTrailingZeros(BufferIndices);
+				BufferIndices &= ~(1 << OutputBufferIndex);
+				FCStringAnsi::Snprintf(OutputBufferIdx, 3, "%u", OutputBufferIndex);
+				Defines[TargetDesc.numOptions++] = { "shader_output_buffer_index", OutputBufferIdx };
+				
+				IndirectParamsIndex = FPlatformMath::CountTrailingZeros(BufferIndices);
+				BufferIndices &= ~(1 << IndirectParamsIndex);
+				FCStringAnsi::Snprintf(IndirectParamsIdx, 3, "%u", IndirectParamsIndex);
+				Defines[TargetDesc.numOptions++] = { "indirect_params_buffer_index", IndirectParamsIdx };
+				
+				TESString += FString::Printf(TEXT("// @TessellationHSOutBuffer: %u\n"), OutputBufferIndex);
+				TESString += FString::Printf(TEXT("// @TessellationPatchCountBuffer: %u\n"), IndirectParamsIndex);
+			}
+			
+			char PatchBufferIdx[3];
+			char TessFactorBufferIdx[3];
+			char ShaderInputWGIdx[3];
+			if (Frequency == HSF_HullShader)
+			{
+				PatchBufferIndex = FPlatformMath::CountTrailingZeros(BufferIndices);
+				BufferIndices &= ~(1 << PatchBufferIndex);
+				FCStringAnsi::Snprintf(PatchBufferIdx, 3, "%u", PatchBufferIndex);
+				Defines[TargetDesc.numOptions++] = { "shader_patch_output_buffer_index", PatchBufferIdx };
+				
+				TessFactorBufferIndex = FPlatformMath::CountTrailingZeros(BufferIndices);
+				BufferIndices &= ~(1 << TessFactorBufferIndex);
+				FCStringAnsi::Snprintf(TessFactorBufferIdx, 3, "%u", TessFactorBufferIndex);
+				Defines[TargetDesc.numOptions++] = { "shader_tess_factor_buffer_index", TessFactorBufferIdx };
+				
+				FCStringAnsi::Snprintf(ShaderInputWGIdx, 3, "%u", ShaderInputWGIndex);
+				Defines[TargetDesc.numOptions++] = { "shader_input_wg_index", ShaderInputWGIdx };
+				
+				TESString += FString::Printf(TEXT("// @TessellationHSTFOutBuffer: %u\n"), TessFactorBufferIndex);
+				TESString += FString::Printf(TEXT("// @TessellationControlPointOutBuffer: %u\n"), PatchBufferIndex);
+				TESString += FString::Printf(TEXT("// @TessellationIndexBuffer: %u\n"), ShaderInputWGIndex);
+			}
+			
 			switch (VersionEnum)
 			{
 				case 6:
@@ -1306,6 +1350,10 @@ bool FMetalShaderOutputCooker::Build(TArray<uint8>& OutData)
 			{
 				MetaData += FString::Printf(TEXT("// @SideTable: spvBufferSizeConstants(%d)\n"), SideTableIndex);
 			}
+			if (TESString.Len())
+			{
+				MetaData += TESString;
+			}
 			MetaData += FString::Printf(TEXT("\n\n"));
 			
 			MetalSource = TCHAR_TO_UTF8(*MetaData);
