Parse the HLSL to extract the Hull shader attributes that have no SPIRV equivalent so can't be parsed from the SPIRV. This is likely to be more fragile but there's no real alternative.

diff --git a/Engine/Source/Developer/Apple/MetalShaderFormat/Private/MetalDerivedData.cpp b/Engine/Source/Developer/Apple/MetalShaderFormat/Private/MetalDerivedData.cpp
--- a/Engine/Source/Developer/Apple/MetalShaderFormat/Private/MetalDerivedData.cpp
+++ b/Engine/Source/Developer/Apple/MetalShaderFormat/Private/MetalDerivedData.cpp
@@ -625,6 +625,39 @@ bool FMetalShaderOutputCooker::Build(TArray<uint8>& OutData)
 			FString OUTString;
 			FString WKGString;
 			
+			if(Frequency == HSF_HullShader)
+			{
+				uint32 PatchSize = 0;
+				size_t ThreadSizeIdx = SourceData.find("InputPatch<");
+				check(ThreadSizeIdx != std::string::npos);
+				char const* String = SourceData.c_str() + ThreadSizeIdx;
+				
+#if !PLATFORM_WINDOWS
+				size_t Found = sscanf(String, "InputPatch< %*s ,  %u >", &PatchSize);
+#else
+				size_t Found = sscanf_s(String, "InputPatch< %*s ,  %u >", &PatchSize);
+#endif
+				check(Found == 1);
+				
+				TESString += FString::Printf(TEXT("// @TessellationInputControlPoints: %u\n"), PatchSize);
+				
+				ThreadSizeIdx = SourceData.find("[maxtessfactor(");
+				check(ThreadSizeIdx != std::string::npos);
+				String = SourceData.c_str() + ThreadSizeIdx;
+				
+				uint32 MaxTessFactor = 0;
+#if !PLATFORM_WINDOWS
+				Found = sscanf(String, "[maxtessfactor(%u)]", &MaxTessFactor);
+#else
+				Found = sscanf_s(String, "[maxtessfactor(%u)]", &MaxTessFactor);
+#endif
+				check(Found == 1);
+				
+				TESString += FString::Printf(TEXT("// @TessellationMaxTessFactor: %u\n"), MaxTessFactor);
+				
+				TESString += TEXT("// @TessellationPatchesPerThreadGroup: 1\n");
+			}
+			
 			{
 				Count = 0;
 				SPVRResult = Reflection.EnumerateExecutionModes(&Count, nullptr);
