Proper order for tessellation meta-data emitted when using DXC to compile shaders to Metal.

diff --git a/Engine/Source/Developer/Apple/MetalShaderFormat/Private/MetalDerivedData.cpp b/Engine/Source/Developer/Apple/MetalShaderFormat/Private/MetalDerivedData.cpp
--- a/Engine/Source/Developer/Apple/MetalShaderFormat/Private/MetalDerivedData.cpp
+++ b/Engine/Source/Developer/Apple/MetalShaderFormat/Private/MetalDerivedData.cpp
@@ -519,7 +519,25 @@ bool FMetalShaderOutputCooker::Build(TArray<uint8>& OutData)
 		SourceDesc.numDefines = 0;
 		SourceDesc.defines = nullptr;
 		
-		FString TESString;
+		enum class EMetalTessellationMetadataTags : uint8
+		{
+			TessellationOutputControlPoints,
+			TessellationDomain,
+			TessellationInputControlPoints,
+			TessellationMaxTessFactor,
+			TessellationOutputWinding,
+			TessellationPartitioning,
+			TessellationPatchesPerThreadGroup,
+			TessellationPatchCountBuffer,
+			TessellationIndexBuffer,
+			TessellationHSOutBuffer,
+			TessellationControlPointOutBuffer,
+			TessellationHSTFOutBuffer,
+			TessellationControlPointIndexBuffer,
+			
+			Num
+		};
+		FString TESStrings[(uint8)EMetalTessellationMetadataTags::Num];
         
         switch (Frequency)
         {
@@ -639,7 +657,7 @@ bool FMetalShaderOutputCooker::Build(TArray<uint8>& OutData)
 #endif
 				check(Found == 1);
 				
-				TESString += FString::Printf(TEXT("// @TessellationInputControlPoints: %u\n"), PatchSize);
+				TESStrings[(uint8)EMetalTessellationMetadataTags::TessellationInputControlPoints] = FString::Printf(TEXT("// @TessellationInputControlPoints: %u\n"), PatchSize);
 				
 				ThreadSizeIdx = SourceData.find("[maxtessfactor(");
 				check(ThreadSizeIdx != std::string::npos);
@@ -653,9 +671,9 @@ bool FMetalShaderOutputCooker::Build(TArray<uint8>& OutData)
 #endif
 				check(Found == 1);
 				
-				TESString += FString::Printf(TEXT("// @TessellationMaxTessFactor: %u\n"), MaxTessFactor);
+				TESStrings[(uint8)EMetalTessellationMetadataTags::TessellationMaxTessFactor] = FString::Printf(TEXT("// @TessellationMaxTessFactor: %u\n"), MaxTessFactor);
 				
-				TESString += TEXT("// @TessellationPatchesPerThreadGroup: 1\n");
+				TESStrings[(uint8)EMetalTessellationMetadataTags::TessellationPatchesPerThreadGroup] = TEXT("// @TessellationPatchesPerThreadGroup: 1\n");
 			}
 			
 			{
@@ -683,13 +701,13 @@ bool FMetalShaderOutputCooker::Build(TArray<uint8>& OutData)
 						case SpvExecutionModeTriangles:
 							if (Frequency == HSF_HullShader || Frequency == HSF_DomainShader)
 							{
-								TESString += TEXT("// @TessellationDomain: tri\n");
+								TESStrings[(uint8)EMetalTessellationMetadataTags::TessellationDomain] = TEXT("// @TessellationDomain: tri\n");
 							}
 							break;
 						case SpvExecutionModeQuads:
 							if (Frequency == HSF_HullShader || Frequency == HSF_DomainShader)
 							{
-								TESString += TEXT("// @TessellationDomain: quad\n");
+								TESStrings[(uint8)EMetalTessellationMetadataTags::TessellationDomain] = TEXT("// @TessellationDomain: quad\n");
 							}
 							break;
 						case SpvExecutionModeIsolines:
@@ -702,37 +720,37 @@ bool FMetalShaderOutputCooker::Build(TArray<uint8>& OutData)
 							if (Frequency == HSF_HullShader || Frequency == HSF_DomainShader)
 							{
 								check(Mode->operands_count == 1);
-								TESString += FString::Printf(TEXT("// @TessellationOutputControlPoints: %d\n"), Mode->operands[0]);
+								TESStrings[(uint8)EMetalTessellationMetadataTags::TessellationOutputControlPoints] = FString::Printf(TEXT("// @TessellationOutputControlPoints: %d\n"), Mode->operands[0]);
 							}
 							break;
 						case SpvExecutionModeVertexOrderCw:
 							if (Frequency == HSF_HullShader || Frequency == HSF_DomainShader)
 							{
-								TESString += TEXT("// @TessellationOutputWinding: cw\n");
+								TESStrings[(uint8)EMetalTessellationMetadataTags::TessellationOutputWinding] = TEXT("// @TessellationOutputWinding: cw\n");
 							}
 							break;
 						case SpvExecutionModeVertexOrderCcw:
 							if (Frequency == HSF_HullShader || Frequency == HSF_DomainShader)
 							{
-								TESString += TEXT("// @TessellationOutputWinding: ccw\n");
+								TESStrings[(uint8)EMetalTessellationMetadataTags::TessellationOutputWinding] = TEXT("// @TessellationOutputWinding: ccw\n");
 							}
 							break;
 						case SpvExecutionModeSpacingEqual:
 							if (Frequency == HSF_HullShader || Frequency == HSF_DomainShader)
 							{
-								TESString += TEXT("// @TessellationPartitioning: integer\n");
+								TESStrings[(uint8)EMetalTessellationMetadataTags::TessellationPartitioning] = TEXT("// @TessellationPartitioning: integer\n");
 							}
 							break;
 						case SpvExecutionModeSpacingFractionalEven:
 							if (Frequency == HSF_HullShader || Frequency == HSF_DomainShader)
 							{
-								TESString += TEXT("// @TessellationPartitioning: fractional_even\n");
+								TESStrings[(uint8)EMetalTessellationMetadataTags::TessellationPartitioning] = TEXT("// @TessellationPartitioning: fractional_even\n");
 							}
 							break;
 						case SpvExecutionModeSpacingFractionalOdd:
 							if (Frequency == HSF_HullShader || Frequency == HSF_DomainShader)
 							{
-								TESString += TEXT("// @TessellationPartitioning: fractional_odd\n");
+								TESStrings[(uint8)EMetalTessellationMetadataTags::TessellationPartitioning] = TEXT("// @TessellationPartitioning: fractional_odd\n");
 							}
 							break;
 						case SpvExecutionModeInvocations:
@@ -1309,8 +1327,8 @@ bool FMetalShaderOutputCooker::Build(TArray<uint8>& OutData)
 				FCStringAnsi::Snprintf(IndirectParamsIdx, 3, "%u", IndirectParamsIndex);
 				Defines[TargetDesc.numOptions++] = { "indirect_params_buffer_index", IndirectParamsIdx };
 				
-				TESString += FString::Printf(TEXT("// @TessellationHSOutBuffer: %u\n"), OutputBufferIndex);
-				TESString += FString::Printf(TEXT("// @TessellationPatchCountBuffer: %u\n"), IndirectParamsIndex);
+				TESStrings[(uint8)EMetalTessellationMetadataTags::TessellationHSOutBuffer] = FString::Printf(TEXT("// @TessellationHSOutBuffer: %u\n"), OutputBufferIndex);
+				TESStrings[(uint8)EMetalTessellationMetadataTags::TessellationPatchCountBuffer] = FString::Printf(TEXT("// @TessellationPatchCountBuffer: %u\n"), IndirectParamsIndex);
 			}
 			
 			char PatchBufferIdx[3];
@@ -1331,9 +1349,9 @@ bool FMetalShaderOutputCooker::Build(TArray<uint8>& OutData)
 				FCStringAnsi::Snprintf(ShaderInputWGIdx, 3, "%u", ShaderInputWGIndex);
 				Defines[TargetDesc.numOptions++] = { "shader_input_wg_index", ShaderInputWGIdx };
 				
-				TESString += FString::Printf(TEXT("// @TessellationHSTFOutBuffer: %u\n"), TessFactorBufferIndex);
-				TESString += FString::Printf(TEXT("// @TessellationControlPointOutBuffer: %u\n"), PatchBufferIndex);
-				TESString += FString::Printf(TEXT("// @TessellationIndexBuffer: %u\n"), ShaderInputWGIndex);
+				TESStrings[(uint8)EMetalTessellationMetadataTags::TessellationHSTFOutBuffer] = FString::Printf(TEXT("// @TessellationHSTFOutBuffer: %u\n"), TessFactorBufferIndex);
+				TESStrings[(uint8)EMetalTessellationMetadataTags::TessellationControlPointOutBuffer] = FString::Printf(TEXT("// @TessellationControlPointOutBuffer: %u\n"), PatchBufferIndex);
+				TESStrings[(uint8)EMetalTessellationMetadataTags::TessellationIndexBuffer] = FString::Printf(TEXT("// @TessellationIndexBuffer: %u\n"), ShaderInputWGIndex);
 			}
 			
 			switch (VersionEnum)
@@ -1383,6 +1401,11 @@ bool FMetalShaderOutputCooker::Build(TArray<uint8>& OutData)
 			{
 				MetaData += FString::Printf(TEXT("// @SideTable: spvBufferSizeConstants(%d)\n"), SideTableIndex);
 			}
+			FString TESString;
+			for (uint32 i = 0; i < (uint32)EMetalTessellationMetadataTags::Num; i++)
+			{
+				TESString += TESStrings[i];
+			}
 			if (TESString.Len())
 			{
 				MetaData += TESString;
