Fix the fma() macros when using DXC to compile shaders - extra brackets are needed because the arguments may be expressions.

diff --git a/Engine/Source/Developer/Apple/MetalShaderFormat/Private/MetalDerivedData.cpp b/Engine/Source/Developer/Apple/MetalShaderFormat/Private/MetalDerivedData.cpp
--- a/Engine/Source/Developer/Apple/MetalShaderFormat/Private/MetalDerivedData.cpp
+++ b/Engine/Source/Developer/Apple/MetalShaderFormat/Private/MetalDerivedData.cpp
@@ -1766,16 +1766,16 @@ bool FMetalShaderOutputCooker::Build(TArray<uint8>& OutData)
 				if (Frequency == HSF_PixelShader || Frequency == HSF_ComputeShader)
 				{
 					FMADefine = "#include <metal_stdlib>\n\n"
-								"#define fma(a, b, c) ((a * b) + c)\n\n";
+								"#define fma(a, b, c) (((a) * (b)) + (c))\n\n";
 				}
 				// Plain vertex & domain shaders need only use FMAs on Metal 1.2-2.0
 				else
 				{
 					FMADefine = std::string("#include <metal_stdlib>\n\n"
 									 "#if __METAL_VERSION__ < 120 || __METAL_VERSION__ >= 210\n"
-									 "\t#define fma(a, b, c) ((a * b) + c)\n"
+									 "\t#define fma(a, b, c) (((a) * (b)) + (c))\n"
 									 "#else\n"
-									 "\t#define fma(a, b, c) fma(a, b, c)\n"
+									 "\t#define fma(a, b, c) fma((a), (b), (c))\n"
 									 "#endif\n\n");
 				}
 			
