DR - Fix static analysis warning
#rnx

diff --git a/Engine/Source/Developer/Apple/MetalShaderFormat/Private/MetalDerivedData.cpp b/Engine/Source/Developer/Apple/MetalShaderFormat/Private/MetalDerivedData.cpp
--- a/Engine/Source/Developer/Apple/MetalShaderFormat/Private/MetalDerivedData.cpp
+++ b/Engine/Source/Developer/Apple/MetalShaderFormat/Private/MetalDerivedData.cpp
@@ -1791,7 +1791,7 @@ bool FMetalShaderOutputCooker::Build(TArray<uint8>& OutData)
 				size_t ThreadSizeIdx = SourceData.find("OutputPatch<");
 				check(ThreadSizeIdx != std::string::npos);
 				char const* String = SourceData.c_str() + ThreadSizeIdx;
-				while(!isdigit(*String))
+				while(!isdigit((unsigned char)*String))
 				{
 					String++;
 				}
