Fixed transformation of HLSLcc IR in Metal backend.

#rnx

diff --git a/Engine/Shaders/Public/Platform/Metal/MetalCommon.ush b/Engine/Shaders/Public/Platform/Metal/MetalCommon.ush
--- a/Engine/Shaders/Public/Platform/Metal/MetalCommon.ush
+++ b/Engine/Shaders/Public/Platform/Metal/MetalCommon.ush
@@ -5,8 +5,8 @@
 =============================================================================*/
 
 // Update this GUID to improve shader recompilation for Metal only shaders
-// GUID = 720867B8-6564-417C-B591-37B22978D201
-  
+// GUID = cda708c4-2f08-4eb2-b045-42a269d14ed4
+
 #pragma once
 
 #if VERTEXSHADER && HULLSHADER
diff --git a/Engine/Source/Developer/Apple/MetalShaderFormat/Private/MetalUtils.cpp b/Engine/Source/Developer/Apple/MetalShaderFormat/Private/MetalUtils.cpp
--- a/Engine/Source/Developer/Apple/MetalShaderFormat/Private/MetalUtils.cpp
+++ b/Engine/Source/Developer/Apple/MetalShaderFormat/Private/MetalUtils.cpp
@@ -3460,7 +3460,30 @@ void FMetalCodeBackend::RemovePackedVarReferences(exec_list* ir, _mesa_glsl_pars
 		}
 		else
 		{
-			OldVar->insert_after(NewVar);
+			// New variable declaration including its initialization must be inserted *after* the original variable has been initialized.
+			// So look ahead until we find the <ir_assignment> node that initializes "OldVar".
+			ir_assignment* OldVarInit = nullptr;
+			for (exec_node* Node = OldVar->next; Node && !OldVarInit; Node = Node->next)
+			{
+				ir_instruction* Instr = static_cast<ir_instruction*>(Node);
+				if (Instr->ir_type == ir_type_assignment)
+				{
+					ir_assignment* AssignStmt = static_cast<ir_assignment*>(Instr);
+					if (AssignStmt->lhs->variable_referenced() == OldVar)
+					{
+						OldVarInit = AssignStmt;
+					}
+				}
+			}
+			
+			if (OldVarInit != nullptr)
+			{
+				OldVarInit->insert_after(NewVar);
+			}
+			else
+			{
+				OldVar->insert_after(NewVar);
+			}
 		}
 	}
 
