WIP support for Indirect Argument Buffers when using ShaderConductor to compile shaders. Doesn't handle atomic buffers or the buffer-size table correctly yet, but when it works it should run on any Metal GPU unlike the hlslcc version.

diff --git a/Engine/Source/Developer/Apple/MetalShaderFormat/Private/MetalDerivedData.cpp b/Engine/Source/Developer/Apple/MetalShaderFormat/Private/MetalDerivedData.cpp
--- a/Engine/Source/Developer/Apple/MetalShaderFormat/Private/MetalDerivedData.cpp
+++ b/Engine/Source/Developer/Apple/MetalShaderFormat/Private/MetalDerivedData.cpp
@@ -611,6 +611,7 @@ bool FMetalShaderOutputCooker::Build(TArray<uint8>& OutData)
 		FString EntryPoint = Input.EntryPointName;
 		EHlslShaderFrequency Freq = Frequency;
 		FString ALNString;
+		FString IABString;
 
 		TargetDesc.language = ShaderConductor::ShadingLanguage::SpirV;
 		ShaderConductor::Compiler::ResultDesc Results = ShaderConductor::Compiler::Compile(SourceDesc, Options, TargetDesc);
@@ -635,6 +636,87 @@ bool FMetalShaderOutputCooker::Build(TArray<uint8>& OutData)
 			uint64 TextureIndices = 0xffffffffffffffff;
 			uint16 SamplerIndices = 0xffff;
 			
+			struct FMetalResourceTableEntry : FResourceTableEntry
+			{
+				FString Name;
+				uint32 Size;
+				uint32 SetIndex;
+				bool bUsed;
+			};
+			
+			TArray<FString> TableNames;
+			TMap<FString, TArray<FMetalResourceTableEntry>> IABs;
+			TMap<FString, FMetalResourceTableEntry> ResourceTable;
+			if (VersionEnum >= 5)
+			{
+				for (auto Pair : Input.Environment.ResourceTableLayoutHashes)
+				{
+					TableNames.Add(*Pair.Key);
+				}
+				
+				for (auto Pair : Input.Environment.ResourceTableMap)
+				{
+					const FResourceTableEntry& Entry = Pair.Value;
+					TArray<FMetalResourceTableEntry>& Resources = IABs.FindOrAdd(Entry.UniformBufferName);
+					if (Resources.Num() <= Entry.ResourceIndex)
+					{
+						Resources.SetNum(Entry.ResourceIndex + 1);
+					}
+					FMetalResourceTableEntry NewEntry;
+					NewEntry.UniformBufferName = Entry.UniformBufferName;
+					NewEntry.Type = Entry.Type;
+					NewEntry.ResourceIndex = Entry.ResourceIndex;
+					NewEntry.Name = Pair.Key;
+					NewEntry.Size = 1;
+					NewEntry.bUsed = false;
+					Resources[Entry.ResourceIndex] = NewEntry;
+				}
+				
+				for (uint32 i = 0; i < TableNames.Num(); )
+				{
+					if (!IABs.Contains(TableNames[i]))
+					{
+						TableNames.RemoveAt(i);
+					}
+					else
+					{
+						i++;
+					}
+				}
+				
+				for (auto Pair : IABs)
+				{
+					uint32 Index = 0;
+					for (uint32 i = 0; i < Pair.Value.Num(); i++)
+					{
+						FMetalResourceTableEntry& Entry = Pair.Value[i];
+						switch(Entry.Type)
+						{
+							case UBMT_UAV:
+							case UBMT_RDG_TEXTURE_UAV:
+							case UBMT_RDG_BUFFER_UAV:
+								Entry.ResourceIndex = Index;
+								Entry.Size = 1;
+								Index += 2;
+								break;
+							default:
+								Entry.ResourceIndex = Index;
+								Index++;
+								break;
+						}
+						for (uint32 j = 0; j < TableNames.Num(); j++)
+						{
+							if (Entry.UniformBufferName == TableNames[j])
+							{
+								Entry.SetIndex = j;
+								break;
+							}
+						}
+						ResourceTable.Add(Entry.Name, Entry);
+					}
+				}
+			}
+			
 			FString UAVString;
 			FString SRVString;
 			FString SMPString;
@@ -802,6 +884,8 @@ bool FMetalShaderOutputCooker::Build(TArray<uint8>& OutData)
 			check(SPVRResult == SPV_REFLECT_RESULT_SUCCESS);
 			if (Count > 0)
 			{
+				TArray<SpvReflectDescriptorBinding*> ResourceBindings;
+				TArray<SpvReflectDescriptorBinding*> ArgumentBindings;
 				TArray<SpvReflectDescriptorBinding*> UniformBindings;
 				TArray<SpvReflectDescriptorBinding*> SamplerBindings;
 				TArray<SpvReflectDescriptorBinding*> TextureSRVBindings;
@@ -810,11 +894,22 @@ bool FMetalShaderOutputCooker::Build(TArray<uint8>& OutData)
 				TArray<SpvReflectDescriptorBinding*> TBufferUAVBindings;
 				TArray<SpvReflectDescriptorBinding*> SBufferSRVBindings;
 				TArray<SpvReflectDescriptorBinding*> SBufferUAVBindings;
+				TSet<FString> UsedSets;
 				
 				// Extract all the bindings first so that we process them in order - this lets us assign UAVs before other resources
 				// Which is necessary to match the D3D binding scheme.
 				for (auto const& Binding : Bindings)
 				{
+					if (Binding->resource_type != SPV_REFLECT_RESOURCE_FLAG_CBV && ResourceTable.Contains(UTF8_TO_TCHAR(Binding->name)))
+					{
+						ResourceBindings.Add(Binding);
+						
+						FMetalResourceTableEntry Entry = ResourceTable.FindRef(UTF8_TO_TCHAR(Binding->name));
+						UsedSets.Add(Entry.UniformBufferName);
+						
+						continue;
+					}
+					
 					switch(Binding->resource_type)
 					{
 						case SPV_REFLECT_RESOURCE_FLAG_CBV:
@@ -822,7 +917,14 @@ bool FMetalShaderOutputCooker::Build(TArray<uint8>& OutData)
 							check(Binding->descriptor_type == SPV_REFLECT_DESCRIPTOR_TYPE_UNIFORM_BUFFER);
 							if (Binding->accessed)
 							{
-								UniformBindings.Add(Binding);
+								if (TableNames.Contains(UTF8_TO_TCHAR(Binding->name)))
+								{
+									ArgumentBindings.Add(Binding);
+								}
+								else
+								{
+									UniformBindings.Add(Binding);
+								}
 							}
 							break;
 						}
@@ -914,6 +1016,113 @@ bool FMetalShaderOutputCooker::Build(TArray<uint8>& OutData)
 					}
 				}
 				
+				for (uint32 i = 0; i < TableNames.Num(); )
+				{
+					if (UsedSets.Contains(TableNames[i]))
+					{
+						i++;
+					}
+					else
+					{
+						IABs.Remove(TableNames[i]);
+						TableNames.RemoveAt(i);
+					}
+				}
+				
+				for (uint32 i = 0; i < ArgumentBindings.Num(); )
+				{
+					FString Name = UTF8_TO_TCHAR(ArgumentBindings[i]->name);
+					if (TableNames.Contains(Name))
+					{
+						auto* ResourceArray = IABs.Find(Name);
+						auto const& LastResource = ResourceArray->Last();
+						uint32 ResIndex = LastResource.ResourceIndex + LastResource.Size;
+						uint32 SetIndex = SPV_REFLECT_SET_NUMBER_DONT_CHANGE;
+						for (uint32 j = 0; j < TableNames.Num(); j++)
+						{
+							if (Name == TableNames[j])
+							{
+								SetIndex = j;
+								break;
+							}
+						}
+						
+						FMetalResourceTableEntry Entry;
+						Entry.UniformBufferName = LastResource.UniformBufferName;
+						Entry.Name = Name;
+						Entry.ResourceIndex = ResIndex;
+						Entry.SetIndex = SetIndex;
+						Entry.bUsed = true;
+						
+						ResourceArray->Add(Entry);
+						ResourceTable.Add(Name, Entry);
+						
+						ResourceBindings.Add(ArgumentBindings[i]);
+						
+						i++;
+					}
+					else
+					{
+						UniformBindings.Add(ArgumentBindings[i]);
+						ArgumentBindings.RemoveAt(i);
+					}
+				}
+				
+				for (auto const& Binding : ResourceBindings)
+				{
+					FMetalResourceTableEntry* Entry = ResourceTable.Find(UTF8_TO_TCHAR(Binding->name));
+					
+					for (uint32 j = 0; j < TableNames.Num(); j++)
+					{
+						if (Entry->UniformBufferName == TableNames[j])
+						{
+							Entry->SetIndex = j;
+							break;
+						}
+					}
+					Entry->bUsed = true;
+					
+					auto* ResourceArray = IABs.Find(Entry->UniformBufferName);
+					for (auto& Resource : *ResourceArray)
+					{
+						if (Resource.Name == Entry->Name)
+						{
+							Resource.SetIndex = Entry->SetIndex;
+							Resource.bUsed = true;
+							break;
+						}
+					}
+					
+					SPVRResult = Reflection.ChangeDescriptorBindingNumbers(Binding, Entry->ResourceIndex, Entry->SetIndex + 1);
+					check(SPVRResult == SPV_REFLECT_RESULT_SUCCESS);
+				}
+				
+				for (auto const& Pair : IABs)
+				{
+					FString Name = Pair.Key;
+					auto const& ResourceArray = Pair.Value;
+					if (IABString.Len())
+					{
+						IABString += TEXT(",");
+					}
+					uint32 SetIndex = ResourceArray[0].SetIndex;
+					IABString += FString::Printf(TEXT("%d["), SetIndex);
+					bool bComma = false;
+					for (auto const& Resource : ResourceArray)
+					{
+						if (Resource.bUsed)
+						{
+							if (bComma)
+							{
+								IABString += TEXT(",");
+							}
+							IABString += FString::Printf(TEXT("%u"), Resource.ResourceIndex);
+							bComma = true;
+						}
+					}
+					IABString += TEXT("]");
+				}
+				
 				for (auto const& Binding : TBufferUAVBindings)
 				{
 					check(UAVIndices);
@@ -1596,7 +1805,7 @@ bool FMetalShaderOutputCooker::Build(TArray<uint8>& OutData)
 			FCStringAnsi::Snprintf(BufferIdx, 3, "%d", SideTableIndex);
 			BufferIndices &= ~(1 << SideTableIndex);
 
-			ShaderConductor::MacroDefine Defines[12] = {{"texel_buffer_texture_width", "0"}, {"enforce_storge_buffer_bounds", "1"}, {"buffer_size_buffer_index", BufferIdx}, {nullptr, nullptr}, {nullptr, nullptr}, {nullptr, nullptr}, {nullptr, nullptr}};
+			ShaderConductor::MacroDefine Defines[13] = {{"texel_buffer_texture_width", "0"}, {"enforce_storge_buffer_bounds", "1"}, {"buffer_size_buffer_index", BufferIdx}, {nullptr, nullptr}, {nullptr, nullptr}, {nullptr, nullptr}, {nullptr, nullptr}, {nullptr, nullptr}};
 			TargetDesc.numOptions = 3;
 			TargetDesc.options = &Defines[0];
 			switch(Semantics)
@@ -1695,6 +1904,9 @@ bool FMetalShaderOutputCooker::Build(TArray<uint8>& OutData)
 			{
 				case 6:
 				case 5:
+				{
+					Defines[TargetDesc.numOptions++] = { "argument_buffers", "1" };
+				}
 				case 4:
 				{
 					Defines[TargetDesc.numOptions++] = { "texture_buffer_native", "1" };
@@ -1738,6 +1950,12 @@ bool FMetalShaderOutputCooker::Build(TArray<uint8>& OutData)
 			{
 				MetaData += FString::Printf(TEXT("// @SideTable: spvBufferSizeConstants(%d)\n"), SideTableIndex);
 			}
+			if (IABString.Len())
+			{
+				MetaData += TEXT("// @ArgumentBuffers: ");
+				MetaData += IABString;
+				MetaData += TEXT("\n");
+			}
 			FString TESString;
 			for (uint32 i = 0; i < (uint32)EMetalTessellationMetadataTags::Num; i++)
 			{
