Make sure that the IABs get a unique binding number for Metal, due to the way we need to setup UAVs they can't just be 0-7.

diff --git a/Engine/Source/Developer/Apple/MetalShaderFormat/Private/MetalDerivedData.cpp b/Engine/Source/Developer/Apple/MetalShaderFormat/Private/MetalDerivedData.cpp
--- a/Engine/Source/Developer/Apple/MetalShaderFormat/Private/MetalDerivedData.cpp
+++ b/Engine/Source/Developer/Apple/MetalShaderFormat/Private/MetalDerivedData.cpp
@@ -612,6 +612,7 @@
 		EHlslShaderFrequency Freq = Frequency;
 		FString ALNString;
 		FString IABString;
+		uint32 IABOffsetIndex = 0;
 
 		TargetDesc.language = ShaderConductor::ShadingLanguage::SpirV;
 		ShaderConductor::Compiler::ResultDesc Results = ShaderConductor::Compiler::Compile(SourceDesc, Options, TargetDesc);
@@ -1068,60 +1069,7 @@
 					}
 				}
 				
-				for (auto const& Binding : ResourceBindings)
-				{
-					FMetalResourceTableEntry* Entry = ResourceTable.Find(UTF8_TO_TCHAR(Binding->name));
-					
-					for (uint32 j = 0; j < TableNames.Num(); j++)
-					{
-						if (Entry->UniformBufferName == TableNames[j])
-						{
-							Entry->SetIndex = j;
-							break;
-						}
-					}
-					Entry->bUsed = true;
-					
-					auto* ResourceArray = IABs.Find(Entry->UniformBufferName);
-					for (auto& Resource : *ResourceArray)
-					{
-						if (Resource.Name == Entry->Name)
-						{
-							Resource.SetIndex = Entry->SetIndex;
-							Resource.bUsed = true;
-							break;
-						}
-					}
-					
-					SPVRResult = Reflection.ChangeDescriptorBindingNumbers(Binding, Entry->ResourceIndex, Entry->SetIndex + 1);
-					check(SPVRResult == SPV_REFLECT_RESULT_SUCCESS);
-				}
+				uint32 GlobalSetId = 32;
 				
-				for (auto const& Pair : IABs)
-				{
-					FString Name = Pair.Key;
-					auto const& ResourceArray = Pair.Value;
-					if (IABString.Len())
-					{
-						IABString += TEXT(",");
-					}
-					uint32 SetIndex = ResourceArray[0].SetIndex;
-					IABString += FString::Printf(TEXT("%d["), SetIndex);
-					bool bComma = false;
-					for (auto const& Resource : ResourceArray)
-					{
-						if (Resource.bUsed)
-						{
-							if (bComma)
-							{
-								IABString += TEXT(",");
-							}
-							IABString += FString::Printf(TEXT("%u"), Resource.ResourceIndex);
-							bComma = true;
-						}
-					}
-					IABString += TEXT("]");
-				}
-				
 				for (auto const& Binding : TBufferUAVBindings)
 				{
@@ -1138,7 +1086,7 @@
 					
 					UAVString += FString::Printf(TEXT("%s%s(%u:%u)"), UAVString.Len() ? TEXT(",") : TEXT(""), UTF8_TO_TCHAR(Binding->name), Index, 1);
 					
-					SPVRResult = Reflection.ChangeDescriptorBindingNumbers(Binding, Index);
+					SPVRResult = Reflection.ChangeDescriptorBindingNumbers(Binding, Index, GlobalSetId);
 					check(SPVRResult == SPV_REFLECT_RESULT_SUCCESS);
 				}
 				
@@ -1156,7 +1104,7 @@
 					
 					UAVString += FString::Printf(TEXT("%s%s(%u:%u)"), UAVString.Len() ? TEXT(",") : TEXT(""), UTF8_TO_TCHAR(Binding->name), Index, 1);
 					
-					SPVRResult = Reflection.ChangeDescriptorBindingNumbers(Binding, Index);
+					SPVRResult = Reflection.ChangeDescriptorBindingNumbers(Binding, Index, GlobalSetId);
 					check(SPVRResult == SPV_REFLECT_RESULT_SUCCESS);
 				}
 				
@@ -1173,10 +1121,69 @@
 					
 					UAVString += FString::Printf(TEXT("%s%s(%u:%u)"), UAVString.Len() ? TEXT(",") : TEXT(""), UTF8_TO_TCHAR(Binding->name), Index, 1);
 					
-					SPVRResult = Reflection.ChangeDescriptorBindingNumbers(Binding, Index);
+					SPVRResult = Reflection.ChangeDescriptorBindingNumbers(Binding, Index, GlobalSetId);
 					check(SPVRResult == SPV_REFLECT_RESULT_SUCCESS);
 				}
 				
+				IABOffsetIndex = FPlatformMath::CountTrailingZeros(BufferIndices);
+				
+				for (auto const& Binding : ResourceBindings)
+				{
+					FMetalResourceTableEntry* Entry = ResourceTable.Find(UTF8_TO_TCHAR(Binding->name));
+					
+					for (uint32 j = 0; j < TableNames.Num(); j++)
+					{
+						if (Entry->UniformBufferName == TableNames[j])
+						{
+							Entry->SetIndex = j + IABOffsetIndex;
+							BufferIndices &= ~(1 << (j + IABOffsetIndex));
+							TextureIndices &= ~(1llu << uint64(j + IABOffsetIndex));
+							break;
+						}
+					}
+					Entry->bUsed = true;
+					
+					auto* ResourceArray = IABs.Find(Entry->UniformBufferName);
+					for (auto& Resource : *ResourceArray)
+					{
+						if (Resource.Name == Entry->Name)
+						{
+							Resource.SetIndex = Entry->SetIndex;
+							Resource.bUsed = true;
+							break;
+						}
+					}
+					
+					SPVRResult = Reflection.ChangeDescriptorBindingNumbers(Binding, Entry->ResourceIndex, Entry->SetIndex);
+					check(SPVRResult == SPV_REFLECT_RESULT_SUCCESS);
+				}
+				
+				for (auto const& Pair : IABs)
+				{
+					FString Name = Pair.Key;
+					auto const& ResourceArray = Pair.Value;
+					if (IABString.Len())
+					{
+						IABString += TEXT(",");
+					}
+					uint32 SetIndex = ResourceArray[0].SetIndex + IABOffsetIndex;
+					IABString += FString::Printf(TEXT("%d["), SetIndex);
+					bool bComma = false;
+					for (auto const& Resource : ResourceArray)
+					{
+						if (Resource.bUsed)
+						{
+							if (bComma)
+							{
+								IABString += TEXT(",");
+							}
+							IABString += FString::Printf(TEXT("%u"), Resource.ResourceIndex);
+							bComma = true;
+						}
+					}
+					IABString += TEXT("]");
+				}
+				
 				for (auto const& Binding : TBufferSRVBindings)
 				{
 					check(TextureIndices);
@@ -1190,7 +1197,7 @@
 					
 					SRVString += FString::Printf(TEXT("%s%s(%u:%u)"), SRVString.Len() ? TEXT(",") : TEXT(""), UTF8_TO_TCHAR(Binding->name), Index, 1);
 					
-					SPVRResult = Reflection.ChangeDescriptorBindingNumbers(Binding, Index);
+					SPVRResult = Reflection.ChangeDescriptorBindingNumbers(Binding, Index, GlobalSetId);
 					check(SPVRResult == SPV_REFLECT_RESULT_SUCCESS);
 				}
 				
@@ -1205,7 +1212,7 @@
 					
 					SRVString += FString::Printf(TEXT("%s%s(%u:%u)"), SRVString.Len() ? TEXT(",") : TEXT(""), UTF8_TO_TCHAR(Binding->name), Index, 1);
 					
-					SPVRResult = Reflection.ChangeDescriptorBindingNumbers(Binding, Index);
+					SPVRResult = Reflection.ChangeDescriptorBindingNumbers(Binding, Index, GlobalSetId);
 					check(SPVRResult == SPV_REFLECT_RESULT_SUCCESS);
 				}
 				
@@ -1269,7 +1276,7 @@
 						UBOString += FString::Printf(TEXT("%s%s(%u)"), UBOString.Len() ? TEXT(",") : TEXT(""), UTF8_TO_TCHAR(Binding->name), Index);
 					}
 					
-					SPVRResult = Reflection.ChangeDescriptorBindingNumbers(Binding, Index);
+					SPVRResult = Reflection.ChangeDescriptorBindingNumbers(Binding, Index, GlobalSetId);
 					check(SPVRResult == SPV_REFLECT_RESULT_SUCCESS);
 				}
 				
@@ -1281,7 +1288,7 @@
 					
 					SRVString += FString::Printf(TEXT("%s%s(%u:%u)"), SRVString.Len() ? TEXT(",") : TEXT(""), UTF8_TO_TCHAR(Binding->name), Index, 1);
 					
-					SPVRResult = Reflection.ChangeDescriptorBindingNumbers(Binding, Index);
+					SPVRResult = Reflection.ChangeDescriptorBindingNumbers(Binding, Index, GlobalSetId);
 					check(SPVRResult == SPV_REFLECT_RESULT_SUCCESS);
 				}
 				
@@ -1293,7 +1300,7 @@
 					
 					SMPString += FString::Printf(TEXT("%s%u:%s"), SMPString.Len() ? TEXT(",") : TEXT(""), Index, UTF8_TO_TCHAR(Binding->name));
 					
-					SPVRResult = Reflection.ChangeDescriptorBindingNumbers(Binding, Index);
+					SPVRResult = Reflection.ChangeDescriptorBindingNumbers(Binding, Index, GlobalSetId);
 					check(SPVRResult == SPV_REFLECT_RESULT_SUCCESS);
 				}
 			}
@@ -1805,7 +1812,7 @@
 			FCStringAnsi::Snprintf(BufferIdx, 3, "%d", SideTableIndex);
 			BufferIndices &= ~(1 << SideTableIndex);
 
-			ShaderConductor::MacroDefine Defines[13] = {{"texel_buffer_texture_width", "0"}, {"enforce_storge_buffer_bounds", "1"}, {"buffer_size_buffer_index", BufferIdx}, {nullptr, nullptr}, {nullptr, nullptr}, {nullptr, nullptr}, {nullptr, nullptr}, {nullptr, nullptr}};
+			ShaderConductor::MacroDefine Defines[14] = {{"texel_buffer_texture_width", "0"}, {"enforce_storge_buffer_bounds", "1"}, {"buffer_size_buffer_index", BufferIdx}, {nullptr, nullptr}, {nullptr, nullptr}, {nullptr, nullptr}, {nullptr, nullptr}, {nullptr, nullptr}, {nullptr, nullptr}};
 			TargetDesc.numOptions = 3;
 			TargetDesc.options = &Defines[0];
 			switch(Semantics)
@@ -1900,12 +1907,15 @@
 				TESStrings[(uint8)EMetalTessellationMetadataTags::TessellationIndexBuffer] = FString::Printf(TEXT("// @TessellationIndexBuffer: %u\n"), HullIndexBuffer);
 			}
 			
+			char ArgumentBufferOffset[3];
 			switch (VersionEnum)
 			{
 				case 6:
 				case 5:
 				{
+					FCStringAnsi::Snprintf(ArgumentBufferOffset, 3, "%u", IABOffsetIndex);
 					Defines[TargetDesc.numOptions++] = { "argument_buffers", "1" };
+					Defines[TargetDesc.numOptions++] = { "argument_buffer_offset", ArgumentBufferOffset };
 				}
 				case 4:
 				{
