Small bug fix to IAB assignment.

diff --git a/Engine/Source/Developer/Apple/MetalShaderFormat/Private/MetalDerivedData.cpp b/Engine/Source/Developer/Apple/MetalShaderFormat/Private/MetalDerivedData.cpp
--- a/Engine/Source/Developer/Apple/MetalShaderFormat/Private/MetalDerivedData.cpp
+++ b/Engine/Source/Developer/Apple/MetalShaderFormat/Private/MetalDerivedData.cpp
@@ -1135,7 +1135,7 @@ bool FMetalShaderOutputCooker::Build(TArray<uint8>& OutData)
 					{
 						if (Entry->UniformBufferName == TableNames[j])
 						{
-							Entry->SetIndex = j + IABOffsetIndex;
+							Entry->SetIndex = j;
 							BufferIndices &= ~(1 << (j + IABOffsetIndex));
 							TextureIndices &= ~(1llu << uint64(j + IABOffsetIndex));
 							break;
