Small compile fixes for MSVC.

diff --git a/Engine/Source/Developer/Apple/MetalShaderFormat/Private/MetalDerivedData.cpp b/Engine/Source/Developer/Apple/MetalShaderFormat/Private/MetalDerivedData.cpp
--- a/Engine/Source/Developer/Apple/MetalShaderFormat/Private/MetalDerivedData.cpp
+++ b/Engine/Source/Developer/Apple/MetalShaderFormat/Private/MetalDerivedData.cpp
@@ -659,7 +659,7 @@ bool FMetalShaderOutputCooker::Build(TArray<uint8>& OutData)
 				{
 					const FResourceTableEntry& Entry = Pair.Value;
 					TArray<FMetalResourceTableEntry>& Resources = IABs.FindOrAdd(Entry.UniformBufferName);
-					if (Resources.Num() <= Entry.ResourceIndex)
+					if ((uint32)Resources.Num() <= Entry.ResourceIndex)
 					{
 						Resources.SetNum(Entry.ResourceIndex + 1);
 					}
@@ -673,7 +673,7 @@ bool FMetalShaderOutputCooker::Build(TArray<uint8>& OutData)
 					Resources[Entry.ResourceIndex] = NewEntry;
 				}
 				
-				for (uint32 i = 0; i < TableNames.Num(); )
+				for (uint32 i = 0; i < (uint32)TableNames.Num(); )
 				{
 					if (!IABs.Contains(TableNames[i]))
 					{
@@ -688,7 +688,7 @@ bool FMetalShaderOutputCooker::Build(TArray<uint8>& OutData)
 				for (auto Pair : IABs)
 				{
 					uint32 Index = 0;
-					for (uint32 i = 0; i < Pair.Value.Num(); i++)
+					for (uint32 i = 0; i < (uint32)Pair.Value.Num(); i++)
 					{
 						FMetalResourceTableEntry& Entry = Pair.Value[i];
 						switch(Entry.Type)
@@ -705,7 +705,7 @@ bool FMetalShaderOutputCooker::Build(TArray<uint8>& OutData)
 								Index++;
 								break;
 						}
-						for (uint32 j = 0; j < TableNames.Num(); j++)
+						for (uint32 j = 0; j < (uint32)TableNames.Num(); j++)
 						{
 							if (Entry.UniformBufferName == TableNames[j])
 							{
@@ -1017,7 +1017,7 @@ bool FMetalShaderOutputCooker::Build(TArray<uint8>& OutData)
 					}
 				}
 				
-				for (uint32 i = 0; i < TableNames.Num(); )
+				for (uint32 i = 0; i < (uint32)TableNames.Num(); )
 				{
 					if (UsedSets.Contains(TableNames[i]))
 					{
@@ -1030,7 +1030,7 @@ bool FMetalShaderOutputCooker::Build(TArray<uint8>& OutData)
 					}
 				}
 				
-				for (uint32 i = 0; i < ArgumentBindings.Num(); )
+				for (uint32 i = 0; i < (uint32)ArgumentBindings.Num(); )
 				{
 					FString Name = UTF8_TO_TCHAR(ArgumentBindings[i]->name);
 					if (TableNames.Contains(Name))
@@ -1039,7 +1039,7 @@ bool FMetalShaderOutputCooker::Build(TArray<uint8>& OutData)
 						auto const& LastResource = ResourceArray->Last();
 						uint32 ResIndex = LastResource.ResourceIndex + LastResource.Size;
 						uint32 SetIndex = SPV_REFLECT_SET_NUMBER_DONT_CHANGE;
-						for (uint32 j = 0; j < TableNames.Num(); j++)
+						for (uint32 j = 0; j < (uint32)TableNames.Num(); j++)
 						{
 							if (Name == TableNames[j])
 							{
@@ -1131,7 +1131,7 @@ bool FMetalShaderOutputCooker::Build(TArray<uint8>& OutData)
 				{
 					FMetalResourceTableEntry* Entry = ResourceTable.Find(UTF8_TO_TCHAR(Binding->name));
 					
-					for (uint32 j = 0; j < TableNames.Num(); j++)
+					for (uint32 j = 0; j < (uint32)TableNames.Num(); j++)
 					{
 						if (Entry->UniformBufferName == TableNames[j])
 						{
