Implement support for Tier 2 style IABs when compiling Metal shaders with DXC. This requires some string manipulation because the compiler can't cope with the UniformBuffer struct initializers (SPIRV doesn't permit StructuredBuffers to be members of structs) so unused entries are removed from the SPIRV & Metal output.

diff --git a/Engine/Source/Developer/Apple/MetalShaderFormat/Private/MetalDerivedData.cpp b/Engine/Source/Developer/Apple/MetalShaderFormat/Private/MetalDerivedData.cpp
--- a/Engine/Source/Developer/Apple/MetalShaderFormat/Private/MetalDerivedData.cpp
+++ b/Engine/Source/Developer/Apple/MetalShaderFormat/Private/MetalDerivedData.cpp
@@ -428,6 +428,15 @@ bool FMetalShaderOutputCooker::Build(TArray<uint8>& OutData)
 	uint32 SourceLen = 0;
 	int32 Result = 0;
 	
+	struct FMetalResourceTableEntry : FResourceTableEntry
+	{
+		FString Name;
+		uint32 Size;
+		uint32 SetIndex;
+		bool bUsed;
+	};
+	TMap<FString, TArray<FMetalResourceTableEntry>> IABs;
+	
 	FString const* UsingTessellationDefine = Input.Environment.GetDefinitions().Find(TEXT("USING_TESSELLATION"));
 	bool bUsingTessellation = ((UsingTessellationDefine != nullptr && FString("1") == *UsingTessellationDefine && Frequency == HSF_VertexShader) || Frequency == HSF_HullShader || Frequency == HSF_DomainShader);
 	
@@ -637,16 +646,7 @@ bool FMetalShaderOutputCooker::Build(TArray<uint8>& OutData)
 			uint64 TextureIndices = 0xffffffffffffffff;
 			uint16 SamplerIndices = 0xffff;
 			
-			struct FMetalResourceTableEntry : FResourceTableEntry
-			{
-				FString Name;
-				uint32 Size;
-				uint32 SetIndex;
-				bool bUsed;
-			};
-			
 			TArray<FString> TableNames;
-			TMap<FString, TArray<FMetalResourceTableEntry>> IABs;
 			TMap<FString, FMetalResourceTableEntry> ResourceTable;
 			if (VersionEnum >= 5)
 			{
@@ -1021,6 +1021,7 @@ bool FMetalShaderOutputCooker::Build(TArray<uint8>& OutData)
 				{
 					if (UsedSets.Contains(TableNames[i]))
 					{
+						IABs.FindChecked(TableNames[i])[0].SetIndex = i;
 						i++;
 					}
 					else
@@ -2070,6 +2071,205 @@ bool FMetalShaderOutputCooker::Build(TArray<uint8>& OutData)
 			} while(Pos != std::string::npos);
 		}
 		
+		// Version 6 means Tier 2 IABs for now.
+		if (VersionEnum >= 6)
+		{
+			char BufferIdx[3];
+			for (auto& IAB : IABs)
+			{
+				uint32 Index = IAB.Value[0].SetIndex;
+				FMemory::Memzero(BufferIdx);
+				FCStringAnsi::Snprintf(BufferIdx, 3, "%d", Index);
+				std::string find_str = "struct spvDescriptorSetBuffer";
+				find_str += BufferIdx;
+				size_t Pos = MetalSource.find(find_str);
+				if (Pos != std::string::npos)
+				{
+					size_t StartPos = MetalSource.find("{", Pos);
+					size_t EndPos = MetalSource.find("}", StartPos);
+					std::string IABName(TCHAR_TO_UTF8(*IAB.Key));
+					size_t UBPos = MetalSource.find("constant type_" + IABName + "*");
+					
+					std::string Declaration = find_str + "\n{\n\tconstant uint* spvBufferSizeConstants [[id(0)]];\n";
+					for (FMetalResourceTableEntry& Entry : IAB.Value)
+					{
+						std::string EntryString;
+						std::string Name(TCHAR_TO_UTF8(*Entry.Name));
+						switch(Entry.Type)
+						{
+							case UBMT_TEXTURE:
+							case UBMT_RDG_TEXTURE:
+							case UBMT_RDG_TEXTURE_SRV:
+							case UBMT_SRV:
+							case UBMT_SAMPLER:
+							case UBMT_RDG_BUFFER:
+							case UBMT_RDG_BUFFER_SRV:
+							case UBMT_UAV:
+							case UBMT_RDG_TEXTURE_UAV:
+							case UBMT_RDG_BUFFER_UAV:
+							{
+								size_t EntryPos = MetalSource.find(Name + " [[id(");
+								if (EntryPos != std::string::npos)
+								{
+									while(MetalSource[--EntryPos] != '\n') {}
+									while(MetalSource[++EntryPos] != '\n')
+									{
+										EntryString += MetalSource[EntryPos];
+									}
+									EntryString += "\n";
+								}
+								else
+								{
+									switch(Entry.Type)
+									{
+										case UBMT_TEXTURE:
+										case UBMT_RDG_TEXTURE:
+										case UBMT_RDG_TEXTURE_SRV:
+										case UBMT_SRV:
+										{
+											std::string typeName = "texture_buffer<float, access::read>";
+											int32 NameIndex = PreprocessedShader.Find(Entry.Name + ";");
+											int32 DeclIndex = NameIndex;
+											if (DeclIndex > 0)
+											{
+												while(PreprocessedShader[--DeclIndex] != TEXT('\n')) {}
+												FString Decl = PreprocessedShader.Mid(DeclIndex, NameIndex - DeclIndex);
+												TCHAR const* Types[] = { TEXT("ByteAddressBuffer<"), TEXT("StructuredBuffer<"), TEXT("Buffer<"), TEXT("Texture2DArray"), TEXT("TextureCubeArray"), TEXT("Texture2D"), TEXT("Texture3D"), TEXT("TextureCube") };
+												char const* NewTypes[] = { "device void*", "device void*", "texture_buffer<float, access::read>", "texture2d_array<float>", "texturecube_array<float>", "texture2d<float>", "texture3d<float>", "texturecube<float>" };
+												for (uint32 i = 0; i < 8; i++)
+												{
+													if (Decl.Contains(Types[i]))
+													{
+														typeName = NewTypes[i];
+														break;
+													}
+												}
+											}
+											
+											FCStringAnsi::Snprintf(BufferIdx, 3, "%d", Entry.ResourceIndex + 1);
+											EntryString = "\t";
+											EntryString += typeName;
+											EntryString += " ";
+											EntryString += Name;
+											EntryString += " [[id(";
+											EntryString += BufferIdx;
+											EntryString += ")]];\n";
+											break;
+										}
+										case UBMT_SAMPLER:
+										{
+											FCStringAnsi::Snprintf(BufferIdx, 3, "%d", Entry.ResourceIndex + 1);
+											EntryString = "\tsampler ";
+											EntryString += Name;
+											EntryString += " [[id(";
+											EntryString += BufferIdx;
+											EntryString += ")]];\n";
+											break;
+										}
+										case UBMT_RDG_BUFFER:
+										case UBMT_RDG_BUFFER_SRV:
+										{
+											FCStringAnsi::Snprintf(BufferIdx, 3, "%d", Entry.ResourceIndex + 1);
+											EntryString = "\tdevice void* ";
+											EntryString += Name;
+											EntryString += " [[id(";
+											EntryString += BufferIdx;
+											EntryString += ")]];\n";
+											break;
+										}
+										case UBMT_UAV:
+										case UBMT_RDG_TEXTURE_UAV:
+										{
+											std::string typeName = "texture_buffer<float, access::read_write>";
+											int32 NameIndex = PreprocessedShader.Find(Entry.Name + ";");
+											int32 DeclIndex = NameIndex;
+											if (DeclIndex > 0)
+											{
+												while(PreprocessedShader[--DeclIndex] != TEXT('\n')) {}
+												FString Decl = PreprocessedShader.Mid(DeclIndex, NameIndex - DeclIndex);
+												TCHAR const* Types[] = { TEXT("ByteAddressBuffer<"), TEXT("StructuredBuffer<"), TEXT("Buffer<"), TEXT("Texture2DArray"), TEXT("TextureCubeArray"), TEXT("Texture2D"), TEXT("Texture3D"), TEXT("TextureCube") };
+												char const* NewTypes[] = { "device void*", "device void*", "texture_buffer<float, access::read_write>", "texture2d_array<float, access::read_write>", "texturecube_array<float, access::read_write>", "texture2d<float, access::read_write>", "texture3d<float, access::read_write>", "texturecube<float, access::read_write>" };
+												for (uint32 i = 0; i < 8; i++)
+												{
+													if (Decl.Contains(Types[i]))
+													{
+														typeName = NewTypes[i];
+														break;
+													}
+												}
+											}
+											
+											FCStringAnsi::Snprintf(BufferIdx, 3, "%d", Entry.ResourceIndex + 1);
+											EntryString = "\t";
+											EntryString += typeName;
+											EntryString += " ";
+											EntryString += Name;
+											EntryString += " [[id(";
+											EntryString += BufferIdx;
+											EntryString += ")]];\n";
+											
+											FCStringAnsi::Snprintf(BufferIdx, 3, "%d", Entry.ResourceIndex + 2);
+											EntryString = "\tdevice void* ";
+											EntryString += Name;
+											EntryString += "_atomic [[id(";
+											EntryString += BufferIdx;
+											EntryString += ")]];\n";
+											break;
+										}
+										case UBMT_RDG_BUFFER_UAV:
+										{
+											FCStringAnsi::Snprintf(BufferIdx, 3, "%d", Entry.ResourceIndex + 1);
+											EntryString = "\ttexture_buffer<float, access::read_write> ";
+											EntryString += Name;
+											EntryString += " [[id(";
+											EntryString += BufferIdx;
+											EntryString += ")]];\n";
+											
+											FCStringAnsi::Snprintf(BufferIdx, 3, "%d", Entry.ResourceIndex + 2);
+											EntryString = "\tdevice void* ";
+											EntryString += Name;
+											EntryString += "_atomic [[id(";
+											EntryString += BufferIdx;
+											EntryString += ")]];\n";
+											break;
+										}
+										default:
+											break;
+									}
+								}
+								Declaration += EntryString;
+								break;
+							}
+							default:
+							{
+								break;
+							}
+						}
+					}
+					if (UBPos < EndPos)
+					{
+						size_t UBEnd = MetalSource.find(";", UBPos);
+						std::string UBStr = MetalSource.substr(UBPos, (UBEnd - UBPos));
+						Declaration += "\t";
+						Declaration += UBStr;
+						Declaration += ";\n";
+					}
+					else
+					{
+						Declaration += "\tconstant void* uniformdata [[id(";
+						FMemory::Memzero(BufferIdx);
+						FCStringAnsi::Snprintf(BufferIdx, 3, "%d", IAB.Value.Num() + 1);
+						Declaration += BufferIdx;
+						Declaration += ")]];\n";
+					}
+					
+					Declaration += "}";
+					
+ 					MetalSource.replace(Pos, (EndPos - Pos) + 1, Declaration);
+				}
+			}
+		}
+		
 		if (Results.errorWarningMsg)
 		{
 			std::string ErrorText((const char*)Results.errorWarningMsg->Data(), Results.errorWarningMsg->Size());
