Support generating Tier1 style IABs which must have packed indices and only include elements that are actually referenced in the shader to stay below the device limits.

diff --git a/Engine/Source/Developer/Apple/MetalShaderFormat/Private/MetalDerivedData.cpp b/Engine/Source/Developer/Apple/MetalShaderFormat/Private/MetalDerivedData.cpp
--- a/Engine/Source/Developer/Apple/MetalShaderFormat/Private/MetalDerivedData.cpp
+++ b/Engine/Source/Developer/Apple/MetalShaderFormat/Private/MetalDerivedData.cpp
@@ -332,7 +332,7 @@ FMetalShaderOutputCooker::FMetalShaderOutputCooker(const FShaderCompilerInput& _
 	, MinOSVersion(_MinOSVersion)
 {
 	FString const* IABVersion = Input.Environment.GetDefinitions().Find(TEXT("METAL_INDIRECT_ARGUMENT_BUFFERS"));
-	if (IABVersion)
+	if (IABVersion && VersionEnum >= 4)
 	{
 		if(IABVersion->IsNumeric())
 		{
@@ -1137,6 +1137,45 @@ bool FMetalShaderOutputCooker::Build(TArray<uint8>& OutData)
 				
 				IABOffsetIndex = FPlatformMath::CountTrailingZeros(BufferIndices);
 				
+				TMap<FString, uint32> IABTier1Index;
+				if (IABTier == 1)
+				{
+					for (auto const& Binding : ResourceBindings)
+					{
+						FMetalResourceTableEntry* Entry = ResourceTable.Find(UTF8_TO_TCHAR(Binding->name));
+						auto* ResourceArray = IABs.Find(Entry->UniformBufferName);
+						if (!IABTier1Index.Contains(Entry->UniformBufferName))
+						{
+							IABTier1Index.Add(Entry->UniformBufferName, 0);
+						}
+						if (Binding->descriptor_type == SPV_REFLECT_DESCRIPTOR_TYPE_STORAGE_BUFFER)
+						{
+							bool bFoundBufferSizes = false;
+							for (auto& Resource : *ResourceArray)
+							{
+								if (Resource.ResourceIndex == 65535)
+								{
+									bFoundBufferSizes = true;
+									break;
+								}
+							}
+							if (!bFoundBufferSizes)
+							{
+								FMetalResourceTableEntry BufferSizes;
+								BufferSizes.UniformBufferName = Entry->UniformBufferName;
+								BufferSizes.Name = TEXT("BufferSizes");
+								BufferSizes.Type = UBMT_SRV;
+								BufferSizes.ResourceIndex = 65535;
+								BufferSizes.SetIndex = Entry->SetIndex;
+								BufferSizes.Size = 1;
+								BufferSizes.bUsed = true;
+								ResourceArray->Insert(BufferSizes, 0);
+								IABTier1Index[Entry->UniformBufferName] = 1;
+							}
+						}
+					}
+				}
+				
 				for (auto const& Binding : ResourceBindings)
 				{
 					FMetalResourceTableEntry* Entry = ResourceTable.Find(UTF8_TO_TCHAR(Binding->name));
@@ -1154,18 +1193,40 @@ bool FMetalShaderOutputCooker::Build(TArray<uint8>& OutData)
 					Entry->bUsed = true;
 					
 					auto* ResourceArray = IABs.Find(Entry->UniformBufferName);
-					for (auto& Resource : *ResourceArray)
+					uint32 ResourceIndex = Entry->ResourceIndex;
+					if (IABTier == 1)
 					{
-						if (Resource.Name == Entry->Name)
+						for (auto& Resource : *ResourceArray)
 						{
 							Resource.SetIndex = Entry->SetIndex;
-							Resource.bUsed = true;
-							break;
+							if (Resource.ResourceIndex == Entry->ResourceIndex)
+							{
+								uint32& Tier1Index = IABTier1Index.FindChecked(Entry->UniformBufferName);
+								ResourceIndex = Tier1Index++;
+								Resource.bUsed = true;
+								break;
+							}
+						}
+						if (Entry->ResourceIndex != 65535)
+						{
+							SPVRResult = Reflection.ChangeDescriptorBindingNumbers(Binding, ResourceIndex, Entry->SetIndex);
+							check(SPVRResult == SPV_REFLECT_RESULT_SUCCESS);
 						}
 					}
-					
-					SPVRResult = Reflection.ChangeDescriptorBindingNumbers(Binding, Entry->ResourceIndex + 1, Entry->SetIndex);
-					check(SPVRResult == SPV_REFLECT_RESULT_SUCCESS);
+					else
+					{
+						for (auto& Resource : *ResourceArray)
+						{
+							if (Resource.Name == Entry->Name)
+							{
+								Resource.SetIndex = Entry->SetIndex;
+								Resource.bUsed = true;
+								break;
+							}
+						}
+						SPVRResult = Reflection.ChangeDescriptorBindingNumbers(Binding, ResourceIndex, Entry->SetIndex);
+						check(SPVRResult == SPV_REFLECT_RESULT_SUCCESS);
+					}
 				}
 				
 				for (auto const& Pair : IABs)
@@ -1187,7 +1248,7 @@ bool FMetalShaderOutputCooker::Build(TArray<uint8>& OutData)
 							{
 								IABString += TEXT(",");
 							}
-							IABString += FString::Printf(TEXT("%u"), Resource.ResourceIndex + 1);
+							IABString += FString::Printf(TEXT("%u"), (Resource.ResourceIndex == 65535 ? 0 : Resource.ResourceIndex + 1));
 							bComma = true;
 						}
 					}
@@ -1925,7 +1986,7 @@ bool FMetalShaderOutputCooker::Build(TArray<uint8>& OutData)
 				case 5:
 				case 4:
 				{
-					if (IABTier > 1)
+					if (IABTier >= 1)
 					{
 						FCStringAnsi::Snprintf(ArgumentBufferOffset, 3, "%u", IABOffsetIndex);
 						Defines[TargetDesc.numOptions++] = { "argument_buffers", "1" };
