Fix copy-pate bug.

diff --git a/Engine/Source/Developer/Apple/MetalShaderFormat/Private/MetalDerivedData.cpp b/Engine/Source/Developer/Apple/MetalShaderFormat/Private/MetalDerivedData.cpp
--- a/Engine/Source/Developer/Apple/MetalShaderFormat/Private/MetalDerivedData.cpp
+++ b/Engine/Source/Developer/Apple/MetalShaderFormat/Private/MetalDerivedData.cpp
@@ -1224,8 +1224,8 @@ bool FMetalShaderOutputCooker::Build(TArray<uint8>& OutData)
 								break;
 							}
 						}
-						SPVRResult = Reflection.ChangeDescriptorBindingNumbers(Binding, ResourceIndex, Entry->SetIndex);
-						check(SPVRResult == SPV_REFLECT_RESULT_SUCCESS);
+                        SPVRResult = Reflection.ChangeDescriptorBindingNumbers(Binding, Entry->ResourceIndex + 1, Entry->SetIndex);
+                        check(SPVRResult == SPV_REFLECT_RESULT_SUCCESS);
 					}
 				}
 				
