Simplify the fma defines so they don't cause compile errors.

diff --git a/Engine/Source/Developer/Apple/MetalShaderFormat/Private/MetalDerivedData.cpp b/Engine/Source/Developer/Apple/MetalShaderFormat/Private/MetalDerivedData.cpp
--- a/Engine/Source/Developer/Apple/MetalShaderFormat/Private/MetalDerivedData.cpp
+++ b/Engine/Source/Developer/Apple/MetalShaderFormat/Private/MetalDerivedData.cpp
@@ -2066,17 +2066,15 @@ bool FMetalShaderOutputCooker::Build(TArray<uint8>& OutData)
 				// Fragment shaders and compute shaders need not use FMAs.
 				if (Frequency == HSF_PixelShader || Frequency == HSF_ComputeShader)
 				{
-					FMADefine = "#include <metal_stdlib>\n\n"
-								"#define fma(a, b, c) (((a) * (b)) + (c))\n\n";
+					FMADefine = "#include <metal_stdlib>\n\n";
 				}
 				// Plain vertex & domain shaders need only use FMAs on Metal 1.2-2.0
 				else
 				{
 					FMADefine = std::string("#include <metal_stdlib>\n\n"
 									 "#if __METAL_VERSION__ < 120 || __METAL_VERSION__ >= 210\n"
-									 "\t#define fma(a, b, c) (((a) * (b)) + (c))\n"
-									 "#else\n"
-									 "\t#define fma(a, b, c) fma((a), (b), (c))\n"
+									 "\ttemplate<typename T> T ue4_fma(T a, T b, T c) { return (((a) * (b)) + (c)); }\n"
+									 "\t#define fma ue4_fma\n"
 									 "#endif\n\n");
 				}
 			
