Use string replacement to avoid breaking the HLSL rewriter when using iOS framebuffer fetch.

diff --git a/Engine/Shaders/Public/Platform/Metal/MetalCommon.ush b/Engine/Shaders/Public/Platform/Metal/MetalCommon.ush
--- a/Engine/Shaders/Public/Platform/Metal/MetalCommon.ush
+++ b/Engine/Shaders/Public/Platform/Metal/MetalCommon.ush
@@ -5,8 +5,8 @@
 =============================================================================*/
 
 // Update this GUID to improve shader recompilation for Metal only shaders
-// GUID = cda708c4-2f08-4eb2-b045-42a269d14ed4
-
+// GUID = cda708c4-2f08-4eb2-b045-59a269d14ee9
+  
 #pragma once
 
 #if VERTEXSHADER && HULLSHADER
@@ -32,8 +32,8 @@
 #if IOS
 #if COMPILER_HLSLCC > 1
 #if PIXELSHADER
-	[[vk::input_attachment_index(0)]] SubpassInput<float4> gl_LastFragData;
-	#define FramebufferFetchES2() gl_LastFragData.SubpassLoad()
+ 	Texture2D<float4> gl_LastFragData;
+	#define FramebufferFetchES2() gl_LastFragData.Load(uint3(0,0,0),0)
 #else
 	#define FramebufferFetchES2() float4(0,0,0,0)
 #endif
diff --git a/Engine/Source/Developer/Apple/MetalShaderFormat/Private/MetalDerivedData.cpp b/Engine/Source/Developer/Apple/MetalShaderFormat/Private/MetalDerivedData.cpp
--- a/Engine/Source/Developer/Apple/MetalShaderFormat/Private/MetalDerivedData.cpp
+++ b/Engine/Source/Developer/Apple/MetalShaderFormat/Private/MetalDerivedData.cpp
@@ -610,6 +610,17 @@ bool FMetalShaderOutputCooker::Build(TArray<uint8>& OutData)
 			SourceData.clear();
             SourceData.resize(RewriteBlob->Size());
 			FCStringAnsi::Strncpy(const_cast<char*>(SourceData.c_str()), (const char*)RewriteBlob->Data(), RewriteBlob->Size());
+			
+			static const std::string glFragDecl = "Texture2D<float4> gl_LastFragData;";
+			size_t glFragDeclIncludePos = SourceData.find(glFragDecl);
+			if (glFragDeclIncludePos != std::string::npos)
+				SourceData.replace(glFragDeclIncludePos, glFragDecl.length(), "[[vk::input_attachment_index(0)]] SubpassInput<float4> gl_LastFragData;");
+			
+			static const std::string glFragLoad = "gl_LastFragData.Load(uint3(0, 0, 0), 0)";
+			size_t glFragLoadIncludePos = SourceData.find(glFragLoad);
+			if (glFragLoadIncludePos != std::string::npos)
+				SourceData.replace(glFragLoadIncludePos, glFragLoad.length(), "gl_LastFragData.SubpassLoad()");
+			
 			SourceDesc.source = SourceData.c_str();
 			
             if (bDumpDebugInfo)
@@ -623,6 +634,20 @@ bool FMetalShaderOutputCooker::Build(TArray<uint8>& OutData)
                 }
             }
         }
+		else
+		{
+			static const std::string glFragDecl = "Texture2D<float4> gl_LastFragData;";
+			size_t glFragDeclIncludePos = SourceData.find(glFragDecl);
+			if (glFragDeclIncludePos != std::string::npos)
+				SourceData.replace(glFragDeclIncludePos, glFragDecl.length(), "[[vk::input_attachment_index(0)]] SubpassInput<float4> gl_LastFragData;");
+			
+			static const std::string glFragLoad = "gl_LastFragData.Load(uint3(0,0,0),0)";
+			size_t glFragLoadIncludePos = SourceData.find(glFragLoad);
+			if (glFragLoadIncludePos != std::string::npos)
+				SourceData.replace(glFragLoadIncludePos, glFragLoad.length(), "gl_LastFragData.SubpassLoad()");
+			
+			SourceDesc.source = SourceData.c_str();
+		}
 		
 		Options.removeUnusedGlobals = false;
 
@@ -1904,7 +1929,6 @@ bool FMetalShaderOutputCooker::Build(TArray<uint8>& OutData)
 				default:
 					TargetDesc.platform = "iOS";
 					Defines[TargetDesc.numOptions++] = { "ios_use_framebuffer_fetch_subpasses", "1" };
-					TargetDesc.numOptions += 1;
 					break;
 			}
 			
