Don't crash if there are too many samplers. The error will be reported upstream.

diff --git a/Engine/Source/Developer/Apple/MetalShaderFormat/Private/MetalDerivedData.cpp b/Engine/Source/Developer/Apple/MetalShaderFormat/Private/MetalDerivedData.cpp
--- a/Engine/Source/Developer/Apple/MetalShaderFormat/Private/MetalDerivedData.cpp
+++ b/Engine/Source/Developer/Apple/MetalShaderFormat/Private/MetalDerivedData.cpp
@@ -681,7 +681,7 @@ bool FMetalShaderOutputCooker::Build(TArray<uint8>& OutData)
 			
 			uint8 UAVIndices = 0xff;
 			uint64 TextureIndices = 0xffffffffffffffff;
-			uint16 SamplerIndices = 0xffff;
+			uint32 SamplerIndices = 0xffffffff;
 			
 			TArray<FString> TableNames;
 			TMap<FString, FMetalResourceTableEntry> ResourceTable;
