Added translation of GetDimensions for 3D texture type to MetalBackend.

#rnx

diff --git a/Engine/Shaders/Public/ShaderVersion.ush b/Engine/Shaders/Public/ShaderVersion.ush
--- a/Engine/Shaders/Public/ShaderVersion.ush
+++ b/Engine/Shaders/Public/ShaderVersion.ush
@@ -3,6 +3,6 @@
 // in Platform.ush (which should be included in any shader) it allows to invalidate the shader DDC.
 // 
 // If you are merging streams and there is a conflict with this GUID you should make a new GUID rather than taking one or the other.
-// GUID = F389CECFEAE44B34B8D358B9C0287092
+// GUID = 481267BC16B245EE9361CADD4DAC5874
 
 
diff --git a/Engine/Source/Developer/Apple/MetalShaderFormat/Private/MetalBackend.cpp b/Engine/Source/Developer/Apple/MetalShaderFormat/Private/MetalBackend.cpp
--- a/Engine/Source/Developer/Apple/MetalShaderFormat/Private/MetalBackend.cpp
+++ b/Engine/Source/Developer/Apple/MetalShaderFormat/Private/MetalBackend.cpp
@@ -2962,32 +2962,72 @@ protected:
 		}
 		else if ( deref->op == ir_image_dimensions)
 		{
-			// Convert from:
-			//	HLSL
-			//		int w, h;
-			//		T.GetDimensions({lod, }w, h);
-			// GLSL
-			//		ivec2 Temp;
-			//		Temp = textureSize(T{, lod});
-			// Metal
-			//		int2 Temp = int2((int)T.get_width({lod}), (int)T.get_height({lod}));
-			ralloc_asprintf_append(buffer, "int2(");
-			deref->image->accept(this);
-			ralloc_asprintf_append(buffer, ".get_width(");
-			
-			if (deref->image_index)
-			{
-				deref->image_index->accept(this);
+			if (deref->image->type->sampler_dimensionality == GLSL_SAMPLER_DIM_3D)
+			{
+				// Convert from:
+				//	HLSL
+				//		int w, h, d;
+				//		T.GetDimensions({lod, }w, h, d);
+				// GLSL
+				//		ivec3 Temp;
+				//		Temp = textureSize(T{, lod});
+				// Metal
+				//		int3 Temp = int3((int)T.get_width({lod}), (int)T.get_height({lod}), (int)T.get_depth({lod}));
+				ralloc_asprintf_append(buffer, "int3(");
+				deref->image->accept(this);
+				ralloc_asprintf_append(buffer, ".get_width(");
+				
+				if (deref->image_index)
+				{
+					deref->image_index->accept(this);
+				}
+				ralloc_asprintf_append(buffer, "), (int)");
+				
+				deref->image->accept(this);
+				ralloc_asprintf_append(buffer, ".get_height(");
+				if (deref->image_index)
+				{
+					deref->image_index->accept(this);
+				}
+				ralloc_asprintf_append(buffer, "), (int)");
+				
+				deref->image->accept(this);
+				ralloc_asprintf_append(buffer, ".get_depth(");
+				if (deref->image_index)
+				{
+					deref->image_index->accept(this);
+				}
+				ralloc_asprintf_append(buffer, "))");
 			}
-			ralloc_asprintf_append(buffer, "), (int)");
-			
-			deref->image->accept(this);
-			ralloc_asprintf_append(buffer, ".get_height(");
-			if (deref->image_index)
+			else
 			{
-				deref->image_index->accept(this);
+				// Convert from:
+				//	HLSL
+				//		int w, h;
+				//		T.GetDimensions({lod, }w, h);
+				// GLSL
+				//		ivec2 Temp;
+				//		Temp = textureSize(T{, lod});
+				// Metal
+				//		int2 Temp = int2((int)T.get_width({lod}), (int)T.get_height({lod}));
+				ralloc_asprintf_append(buffer, "int2(");
+				deref->image->accept(this);
+				ralloc_asprintf_append(buffer, ".get_width(");
+				
+				if (deref->image_index)
+				{
+					deref->image_index->accept(this);
+				}
+				ralloc_asprintf_append(buffer, "), (int)");
+				
+				deref->image->accept(this);
+				ralloc_asprintf_append(buffer, ".get_height(");
+				if (deref->image_index)
+				{
+					deref->image_index->accept(this);
+				}
+				ralloc_asprintf_append(buffer, "))");
 			}
-			ralloc_asprintf_append(buffer, "))");
 		}
 		else
 		{
