Handle the global parameters being pushed into the push-constant storage type.

diff --git a/Engine/Source/Developer/Apple/MetalShaderFormat/Private/MetalDerivedData.cpp b/Engine/Source/Developer/Apple/MetalShaderFormat/Private/MetalDerivedData.cpp
--- a/Engine/Source/Developer/Apple/MetalShaderFormat/Private/MetalDerivedData.cpp
+++ b/Engine/Source/Developer/Apple/MetalShaderFormat/Private/MetalDerivedData.cpp
@@ -677,6 +677,7 @@ bool FMetalShaderOutputCooker::Build(TArray<uint8>& OutData)
 			TSet<SpvReflectDescriptorBinding*> Counters;
 			TArray<SpvReflectInterfaceVariable*> InputVars;
 			TArray<SpvReflectInterfaceVariable*> OutputVars;
+			TArray<SpvReflectBlockVariable*> ConstantBindings;
 			TArray<SpvReflectExecutionMode*> ExecutionModes;
 			
 			uint8 UAVIndices = 0xff;
@@ -1405,6 +1406,73 @@ bool FMetalShaderOutputCooker::Build(TArray<uint8>& OutData)
 				}
 			}
 			
+			{
+				Count = 0;
+				SPVRResult = Reflection.EnumeratePushConstantBlocks(&Count, nullptr);
+				check(SPVRResult == SPV_REFLECT_RESULT_SUCCESS);
+				ConstantBindings.SetNum(Count);
+				SPVRResult = Reflection.EnumeratePushConstantBlocks(&Count, ConstantBindings.GetData());
+				check(SPVRResult == SPV_REFLECT_RESULT_SUCCESS);
+				if (Count > 0)
+				{
+					for (auto const& Var : ConstantBindings)
+					{
+						// Global uniform buffer - handled specially as we care about the internal layout
+						if (strstr(Var->name, "$Globals"))
+						{
+							check(BufferIndices);
+							uint32 Index = FPlatformMath::CountTrailingZeros(BufferIndices);
+							BufferIndices &= ~(1 << Index);
+							
+							OutputData.ConstantBuffers |= (1 << Index);
+							
+							TCBDMARangeMap CBRanges;
+							GLOString = FString::Printf(TEXT("Globals(%u): "), Index);
+							
+							FString MbrString;
+							for (uint32 i = 0; i < Var->member_count; i++)
+							{
+								SpvReflectBlockVariable& member = Var->members[i];
+								uint32 MbrOffset = member.absolute_offset / sizeof(float);
+								uint32 MbrSize = member.size / sizeof(float);
+								
+								MbrString += FString::Printf(TEXT("%s%s(%u,%u)"), MbrString.Len() ? TEXT(",") : TEXT(""), UTF8_TO_TCHAR(member.name), MbrOffset, MbrSize);
+								
+								unsigned DestCBPrecision = TEXT('h');
+								
+								//							auto const type = *member.type_description;
+								//							uint32_t masked_type = type.type_flags & 0xF;
+								//
+								//							switch (masked_type) {
+								//								default: checkf(false, TEXT("unsupported component type %d"), masked_type); break;
+								//								case SPV_REFLECT_TYPE_FLAG_BOOL  : DestCBPrecision = TEXT('u'); break;
+								//								case SPV_REFLECT_TYPE_FLAG_INT   : DestCBPrecision = (type.traits.numeric.scalar.signedness ? TEXT('i') : TEXT('u')); break;
+								//								case SPV_REFLECT_TYPE_FLAG_FLOAT : DestCBPrecision = (type.traits.numeric.scalar.width == 32 ? TEXT('h') : TEXT('m')); break;
+								//							}
+								
+								unsigned SourceOffset = MbrOffset;
+								unsigned DestOffset = MbrOffset;
+								unsigned DestSize = MbrSize;
+								unsigned DestCBIndex = 0;
+								InsertRange(CBRanges, Index, SourceOffset, DestSize, DestCBIndex, DestCBPrecision, DestOffset);
+							}
+							
+							for (auto Iter = CBRanges.begin(); Iter != CBRanges.end(); ++Iter)
+							{
+								TDMARangeList& List = Iter->second;
+								for (auto IterList = List.begin(); IterList != List.end(); ++IterList)
+								{
+									check(IterList->DestCBIndex == 0);
+									PAKString += FString::Printf(TEXT("%s%u:%u-%c:%u:%u"), PAKString.Len() ? TEXT(",") : TEXT(""), IterList->SourceCB, IterList->SourceOffset, IterList->DestCBPrecision, IterList->DestOffset, IterList->Size);
+								}
+							}
+							
+							GLOString += MbrString;
+						}
+					}
+				}
+			}
+			
 			if (Frequency == HSF_PixelShader)
 			{
 				Count = 0;
