Add missing cast
#rnx

diff --git a/Engine/Source/Developer/Apple/MetalShaderFormat/Private/MetalBackend.cpp b/Engine/Source/Developer/Apple/MetalShaderFormat/Private/MetalBackend.cpp
--- a/Engine/Source/Developer/Apple/MetalShaderFormat/Private/MetalBackend.cpp
+++ b/Engine/Source/Developer/Apple/MetalShaderFormat/Private/MetalBackend.cpp
@@ -975,7 +975,7 @@ protected:
 							if (!bIsStructuredBuffer && !bIsByteAddressBuffer && !bIsAtomic)
 							{
 								Backend.InvariantBuffers |= (1 << BufferIndex);
-								Backend.TypedBufferFormats[BufferIndex] = GetBufferFormat(PtrType->inner_type);
+								Backend.TypedBufferFormats[BufferIndex] = (uint8)GetBufferFormat(PtrType->inner_type);
 							}
 						}
 						ralloc_asprintf_append(buffer, " *%s", unique_name(var));
@@ -992,7 +992,7 @@ protected:
                         ralloc_asprintf_append(buffer, "typedBuffer%d_rw(", PtrType->inner_type->components());
                         print_type_pre(PtrType->inner_type);
                         ralloc_asprintf_append(buffer, ", %s, %d)", unique_name(var), BufferIndex);
-                        Backend.TypedBufferFormats[BufferIndex] = GetBufferFormat(PtrType->inner_type);
+                        Backend.TypedBufferFormats[BufferIndex] = (uint8)GetBufferFormat(PtrType->inner_type);
                         Backend.TypedBuffers |= (1 << BufferIndex);
                         Backend.TypedUAVs |= (1 << BufferIndex);
 					}
@@ -1074,7 +1074,7 @@ protected:
 								if (!bIsStructuredBuffer && !bIsByteAddressBuffer && !bIsAtomic)
 								{
 									Backend.InvariantBuffers |= (1 << BufferIndex);
-									Backend.TypedBufferFormats[BufferIndex] = GetBufferFormat(PtrType->inner_type);
+									Backend.TypedBufferFormats[BufferIndex] = (uint8)GetBufferFormat(PtrType->inner_type);
 								}
 							}
 							else
@@ -1082,7 +1082,7 @@ protected:
                                 ralloc_asprintf_append(buffer, "typedBuffer%d_read(", PtrType->inner_type->components());
                                 print_type_pre(PtrType->inner_type);
                                 ralloc_asprintf_append(buffer, ", %s, %d)", unique_name(var), BufferIndex);
-                                Backend.TypedBufferFormats[BufferIndex] = GetBufferFormat(PtrType->inner_type);
+                                Backend.TypedBufferFormats[BufferIndex] = (uint8)GetBufferFormat(PtrType->inner_type);
                                 Backend.TypedBuffers |= (1 << BufferIndex);
 							}
 						}
