Fix additional EMetalBufferFormat usages
#rnx

diff --git a/Engine/Source/Developer/Apple/MetalShaderFormat/Private/MetalShaderCompiler.cpp b/Engine/Source/Developer/Apple/MetalShaderFormat/Private/MetalShaderCompiler.cpp
--- a/Engine/Source/Developer/Apple/MetalShaderFormat/Private/MetalShaderCompiler.cpp
+++ b/Engine/Source/Developer/Apple/MetalShaderFormat/Private/MetalShaderCompiler.cpp
@@ -1113,11 +1113,11 @@ void BuildMetalShaderOutput(
 		{
 			if ((TypedBuffers & (1 << i)) != 0)
 			{
-				check(TypedBufferFormats[i] > Unknown);
-                check(TypedBufferFormats[i] < Max);
+				check(TypedBufferFormats[i] > EMetalBufferFormat::Unknown);
+                check(TypedBufferFormats[i] < EMetalBufferFormat::Max);
                 if ((TypeMode > EMetalTypeBufferModeRaw)
                 && (TypeMode <= EMetalTypeBufferModeTB)
-                && (TypedBufferFormats[i] < RGB8Sint || TypedBufferFormats[i] > RGB32Float)
+                && (TypedBufferFormats[i] < EMetalBufferFormat::RGB8Sint || TypedBufferFormats[i] > EMetalBufferFormat::RGB32Float)
                 && (TypeMode == EMetalTypeBufferMode2D || TypeMode == EMetalTypeBufferModeTB || !(TypedUAVs & (1 << i))))
                 {
                 	Header.Bindings.LinearBuffer |= (1 << i);
