Support source filename parameters containing spaces.

#rnx

diff --git a/Engine/Source/Developer/Apple/MetalShaderFormat/Private/MetalShaderCompiler.cpp b/Engine/Source/Developer/Apple/MetalShaderFormat/Private/MetalShaderCompiler.cpp
--- a/Engine/Source/Developer/Apple/MetalShaderFormat/Private/MetalShaderCompiler.cpp
+++ b/Engine/Source/Developer/Apple/MetalShaderFormat/Private/MetalShaderCompiler.cpp
@@ -2758,7 +2758,7 @@ bool FinalizeLibrary_Metal(FName const& Format, FString const& WorkingDir, FStri
 				
 				// Build source file name path
 				UE_LOG(LogShaders, Verbose, TEXT("[%d/%d] %s Main_%0.8x_%0.8x.o"), ++Index, Shaders.Num(), *Format.GetPlainNameString(), Len, CRC);
-				FString SourceFileNameParam = FString::Printf(TEXT("%s/Main_%0.8x_%0.8x.o"), *FPaths::ConvertRelativePathToFull(WorkingDir), Len, CRC);
+				FString SourceFileNameParam = FString::Printf(TEXT("\"%s/Main_%0.8x_%0.8x.o\""), *FPaths::ConvertRelativePathToFull(WorkingDir), Len, CRC);
 				
 				// Remote builds copy file and swizzle Source File Name param
 				if(bBuildingRemotely)
