F - Potential fix for bit overflow
#rnx

diff --git a/Engine/Source/Developer/Apple/MetalShaderFormat/Private/MetalDerivedData.cpp b/Engine/Source/Developer/Apple/MetalShaderFormat/Private/MetalDerivedData.cpp
--- a/Engine/Source/Developer/Apple/MetalShaderFormat/Private/MetalDerivedData.cpp
+++ b/Engine/Source/Developer/Apple/MetalShaderFormat/Private/MetalDerivedData.cpp
@@ -374,7 +374,7 @@ FString FMetalShaderOutputCooker::GetPluginSpecificCacheKeySuffix() const
 		uint64 Flags = 0;
 		for (uint32 Flag : Input.Environment.CompilerFlags)
 		{
-			Flags |= (1llu << uint64(Flag));
+			Flags |= (1ull << uint64(Flag));
 		}
 
 		CachedOutputName = FString::Printf(TEXT("%s-%s_%s-%u_%hu_%llu_%hhu_%d_%s_%s"), *Input.ShaderFormat.GetPlainNameString(), *Input.EntryPointName, *Hash.ToString(), Len, FormatVers, Flags, VersionEnum, IABTier, *GUIDHash.ToString(), *Standard);
@@ -672,7 +672,7 @@ bool FMetalShaderOutputCooker::Build(TArray<uint8>& OutData)
 		TargetDesc.language = ShaderConductor::ShadingLanguage::SpirV;
 		ShaderConductor::Compiler::ResultDesc Results = ShaderConductor::Compiler::Compile(SourceDesc, Options, TargetDesc);
 		Result = (Results.hasError) ? 0 : 1;
-		uint32 BufferIndices = 0xffffffff;
+		uint64 BufferIndices = 0xffffffffffffffff;
 		if (!Results.hasError)
 		{
 			// Now perform reflection on the SPIRV and tweak any decorations that we need to.
@@ -691,7 +691,7 @@ bool FMetalShaderOutputCooker::Build(TArray<uint8>& OutData)
 			
 			uint8 UAVIndices = 0xff;
 			uint64 TextureIndices = 0xffffffffffffffff;
-			uint32 SamplerIndices = 0xffffffff;
+			uint64 SamplerIndices = 0xffffffffffffffff;
 			
 			TArray<FString> TableNames;
 			TMap<FString, FMetalResourceTableEntry> ResourceTable;
@@ -1115,8 +1115,8 @@ bool FMetalShaderOutputCooker::Build(TArray<uint8>& OutData)
 					uint32 Index = FPlatformMath::CountTrailingZeros(UAVIndices);
 					
 					// UAVs always claim all slots so we don't have conflicts as D3D expects 0-7
-					BufferIndices &= ~(1 << Index);
-					TextureIndices &= ~(1llu << uint64(Index));
+					BufferIndices &= ~(1ull << (uint64)Index);
+					TextureIndices &= ~(1ull << (uint64)Index);
 					UAVIndices &= ~(1 << Index);
 					
 					OutputData.TypedUAVs |= (1 << Index);
@@ -1134,8 +1134,8 @@ bool FMetalShaderOutputCooker::Build(TArray<uint8>& OutData)
 					uint32 Index = FPlatformMath::CountTrailingZeros(UAVIndices);
 					
 					// UAVs always claim all slots so we don't have conflicts as D3D expects 0-7
-					BufferIndices &= ~(1 << Index);
-					TextureIndices &= ~(1llu << uint64(Index));
+					BufferIndices &= ~(1ull << (uint64)Index);
+					TextureIndices &= ~(1ull << (uint64)Index);
 					UAVIndices &= ~(1 << Index);
 					
 					OutputData.InvariantBuffers |= (1 << Index);
@@ -1153,8 +1153,8 @@ bool FMetalShaderOutputCooker::Build(TArray<uint8>& OutData)
 					
 					// UAVs always claim all slots so we don't have conflicts as D3D expects 0-7
 					// For texture2d this allows us to emulate atomics with buffers
-					BufferIndices &= ~(1 << Index);
-					TextureIndices &= ~(1llu << uint64(Index));
+					BufferIndices &= ~(1ull << (uint64)Index);
+					TextureIndices &= ~(1ull << (uint64)Index);
 					UAVIndices &= ~(1 << Index);
 					
 					UAVString += FString::Printf(TEXT("%s%s(%u:%u)"), UAVString.Len() ? TEXT(",") : TEXT(""), UTF8_TO_TCHAR(Binding->name), Index, 1);
@@ -1163,7 +1163,7 @@ bool FMetalShaderOutputCooker::Build(TArray<uint8>& OutData)
 					check(SPVRResult == SPV_REFLECT_RESULT_SUCCESS);
 				}
 				
-				IABOffsetIndex = FPlatformMath::CountTrailingZeros(BufferIndices);
+				IABOffsetIndex = FPlatformMath::CountTrailingZeros64(BufferIndices);
 				
 				TMap<FString, uint32> IABTier1Index;
 				if (IABTier == 1)
@@ -1213,8 +1213,8 @@ bool FMetalShaderOutputCooker::Build(TArray<uint8>& OutData)
 						if (Entry->UniformBufferName == TableNames[j])
 						{
 							Entry->SetIndex = j;
-							BufferIndices &= ~(1 << (j + IABOffsetIndex));
-							TextureIndices &= ~(1llu << uint64(j + IABOffsetIndex));
+							BufferIndices &= ~(1ull << (uint64)(j + IABOffsetIndex));
+							TextureIndices &= ~(1ull << uint64(j + IABOffsetIndex));
 							break;
 						}
 					}
@@ -1290,8 +1290,8 @@ bool FMetalShaderOutputCooker::Build(TArray<uint8>& OutData)
 					uint32 Index = FPlatformMath::CountTrailingZeros(TextureIndices);
 					
 					// No support for 3-component types in dxc/SPIRV/MSL - need to expose my workarounds there too
-					BufferIndices &= ~(1 << Index);
-					TextureIndices &= ~(1llu << uint64(Index));
+					BufferIndices &= ~(1ull << (uint64)Index);
+					TextureIndices &= ~(1ull << uint64(Index));
 					
 					OutputData.TypedBuffers |= (1 << Index);
 					
@@ -1304,9 +1304,9 @@ bool FMetalShaderOutputCooker::Build(TArray<uint8>& OutData)
 				for (auto const& Binding : SBufferSRVBindings)
 				{
 					check(BufferIndices);
-					uint32 Index = FPlatformMath::CountTrailingZeros(BufferIndices);
+					uint32 Index = FPlatformMath::CountTrailingZeros64(BufferIndices);
 					
-					BufferIndices &= ~(1 << Index);
+					BufferIndices &= ~(1ull << (uint64)Index);
 					
 					OutputData.InvariantBuffers |= (1 << Index);
 					
@@ -1319,8 +1319,8 @@ bool FMetalShaderOutputCooker::Build(TArray<uint8>& OutData)
 				for (auto const& Binding : UniformBindings)
 				{
 					check(BufferIndices);
-					uint32 Index = FPlatformMath::CountTrailingZeros(BufferIndices);
-					BufferIndices &= ~(1 << Index);
+					uint32 Index = FPlatformMath::CountTrailingZeros64(BufferIndices);
+					BufferIndices &= ~(1ull << (uint64)Index);
 					
 					OutputData.ConstantBuffers |= (1 << Index);
 					
@@ -1373,7 +1373,7 @@ bool FMetalShaderOutputCooker::Build(TArray<uint8>& OutData)
 				{
 					check(TextureIndices);
 					uint32 Index = FPlatformMath::CountTrailingZeros64(TextureIndices);
-					TextureIndices &= ~(1llu << uint64(Index));
+					TextureIndices &= ~(1ull << uint64(Index));
 					
 					SRVString += FString::Printf(TEXT("%s%s(%u:%u)"), SRVString.Len() ? TEXT(",") : TEXT(""), UTF8_TO_TCHAR(Binding->name), Index, 1);
 					
@@ -1384,8 +1384,8 @@ bool FMetalShaderOutputCooker::Build(TArray<uint8>& OutData)
 				for (auto const& Binding : SamplerBindings)
 				{
 					check(SamplerIndices);
-					uint32 Index = FPlatformMath::CountTrailingZeros(SamplerIndices);
-					SamplerIndices &= ~(1 << Index);
+					uint32 Index = FPlatformMath::CountTrailingZeros64(SamplerIndices);
+					SamplerIndices &= ~(1ull << (uint64)Index);
 					
 					SMPString += FString::Printf(TEXT("%s%u:%s"), SMPString.Len() ? TEXT(",") : TEXT(""), Index, UTF8_TO_TCHAR(Binding->name));
 					
@@ -1858,10 +1858,10 @@ bool FMetalShaderOutputCooker::Build(TArray<uint8>& OutData)
 			
 			TargetDesc.language = ShaderConductor::ShadingLanguage::Msl;
 			
-			SideTableIndex = FPlatformMath::CountTrailingZeros(BufferIndices);
+			SideTableIndex = FPlatformMath::CountTrailingZeros64(BufferIndices);
 			char BufferIdx[3];
 			FCStringAnsi::Snprintf(BufferIdx, 3, "%d", SideTableIndex);
-			BufferIndices &= ~(1 << SideTableIndex);
+			BufferIndices &= ~(1ull << (uint64)SideTableIndex);
 
 			ShaderConductor::MacroDefine Defines[16] =
 			{
@@ -1902,13 +1902,13 @@ bool FMetalShaderOutputCooker::Build(TArray<uint8>& OutData)
 			char IndirectParamsIdx[3];
 			if (Frequency == HSF_VertexShader && bUsingTessellation)
 			{
-				OutputBufferIndex = FPlatformMath::CountTrailingZeros(BufferIndices);
-				BufferIndices &= ~(1 << OutputBufferIndex);
+				OutputBufferIndex = FPlatformMath::CountTrailingZeros64(BufferIndices);
+				BufferIndices &= ~(1ull << (uint64)OutputBufferIndex);
 				FCStringAnsi::Snprintf(OutputBufferIdx, 3, "%u", OutputBufferIndex);
 				Defines[TargetDesc.numOptions++] = { "shader_output_buffer_index", OutputBufferIdx };
 				
-				IndirectParamsIndex = FPlatformMath::CountTrailingZeros(BufferIndices);
-				BufferIndices &= ~(1 << IndirectParamsIndex);
+				IndirectParamsIndex = FPlatformMath::CountTrailingZeros64(BufferIndices);
+				BufferIndices &= ~(1ull << (uint64)IndirectParamsIndex);
 				FCStringAnsi::Snprintf(IndirectParamsIdx, 3, "%u", IndirectParamsIndex);
 				Defines[TargetDesc.numOptions++] = { "indirect_params_buffer_index", IndirectParamsIdx };
 				
@@ -1918,11 +1918,11 @@ bool FMetalShaderOutputCooker::Build(TArray<uint8>& OutData)
 			
 			if (Frequency == HSF_DomainShader)
 			{
-				OutputBufferIndex = FPlatformMath::CountTrailingZeros(BufferIndices);
-				BufferIndices &= ~(1 << OutputBufferIndex);
+				OutputBufferIndex = FPlatformMath::CountTrailingZeros64(BufferIndices);
+				BufferIndices &= ~(1ull << (uint64)OutputBufferIndex);
 				
-				IndirectParamsIndex = FPlatformMath::CountTrailingZeros(BufferIndices);
-				BufferIndices &= ~(1 << IndirectParamsIndex);
+				IndirectParamsIndex = FPlatformMath::CountTrailingZeros64(BufferIndices);
+				BufferIndices &= ~(1ull << (uint64)IndirectParamsIndex);
 				
 				TESStrings[(uint8)EMetalTessellationMetadataTags::TessellationHSOutBuffer] = FString::Printf(TEXT("// @TessellationHSOutBuffer: %u\n"), OutputBufferIndex);
 				TESStrings[(uint8)EMetalTessellationMetadataTags::TessellationControlPointOutBuffer] = FString::Printf(TEXT("// @TessellationControlPointOutBuffer: %u\n"), IndirectParamsIndex);
@@ -1932,32 +1932,32 @@ bool FMetalShaderOutputCooker::Build(TArray<uint8>& OutData)
 			char TessFactorBufferIdx[3];
 			if (Frequency == HSF_HullShader)
 			{
-				OutputBufferIndex = FPlatformMath::CountTrailingZeros(BufferIndices);
-				BufferIndices &= ~(1 << OutputBufferIndex);
+				OutputBufferIndex = FPlatformMath::CountTrailingZeros64(BufferIndices);
+				BufferIndices &= ~(1ull << (uint64)OutputBufferIndex);
 				FCStringAnsi::Snprintf(OutputBufferIdx, 3, "%u", OutputBufferIndex);
 				Defines[TargetDesc.numOptions++] = { "shader_output_buffer_index", OutputBufferIdx };
 				
-				IndirectParamsIndex = FPlatformMath::CountTrailingZeros(BufferIndices);
-				BufferIndices &= ~(1 << IndirectParamsIndex);
+				IndirectParamsIndex = FPlatformMath::CountTrailingZeros64(BufferIndices);
+				BufferIndices &= ~(1ull << (uint64)IndirectParamsIndex);
 				FCStringAnsi::Snprintf(IndirectParamsIdx, 3, "%u", IndirectParamsIndex);
 				Defines[TargetDesc.numOptions++] = { "indirect_params_buffer_index", IndirectParamsIdx };
 				
-				PatchBufferIndex = FPlatformMath::CountTrailingZeros(BufferIndices);
-				BufferIndices &= ~(1 << PatchBufferIndex);
+				PatchBufferIndex = FPlatformMath::CountTrailingZeros64(BufferIndices);
+				BufferIndices &= ~(1ull << (uint64)PatchBufferIndex);
 				FCStringAnsi::Snprintf(PatchBufferIdx, 3, "%u", PatchBufferIndex);
 				Defines[TargetDesc.numOptions++] = { "shader_patch_output_buffer_index", PatchBufferIdx };
 				
-				TessFactorBufferIndex = FPlatformMath::CountTrailingZeros(BufferIndices);
-				BufferIndices &= ~(1 << TessFactorBufferIndex);
+				TessFactorBufferIndex = FPlatformMath::CountTrailingZeros64(BufferIndices);
+				BufferIndices &= ~(1ull << (uint64)TessFactorBufferIndex);
 				FCStringAnsi::Snprintf(TessFactorBufferIdx, 3, "%u", TessFactorBufferIndex);
 				Defines[TargetDesc.numOptions++] = { "shader_tess_factor_buffer_index", TessFactorBufferIdx };
 				Defines[TargetDesc.numOptions++] = { "shader_input_wg_index", "0" };
 
-				HullIndexBuffer = FPlatformMath::CountTrailingZeros(BufferIndices);
-				BufferIndices &= ~(1 << HullIndexBuffer);
+				HullIndexBuffer = FPlatformMath::CountTrailingZeros64(BufferIndices);
+				BufferIndices &= ~(1ull << (uint64)HullIndexBuffer);
 				
-				uint32 HullVertexBuffer = FPlatformMath::CountTrailingZeros(BufferIndices);
-				BufferIndices &= ~(1 << HullVertexBuffer);
+				uint32 HullVertexBuffer = FPlatformMath::CountTrailingZeros64(BufferIndices);
+				BufferIndices &= ~(1ull << (uint64)HullVertexBuffer);
 				
 				TESStrings[(uint8)EMetalTessellationMetadataTags::TessellationHSOutBuffer] = FString::Printf(TEXT("// @TessellationHSOutBuffer: %u\n"), PatchBufferIndex);
 				TESStrings[(uint8)EMetalTessellationMetadataTags::TessellationPatchCountBuffer] = FString::Printf(TEXT("// @TessellationPatchCountBuffer: %u\n"), IndirectParamsIndex);
