Fix buffer overrun caught by ASAN.

diff --git a/Engine/Source/Developer/Apple/MetalShaderFormat/Private/MetalDerivedData.cpp b/Engine/Source/Developer/Apple/MetalShaderFormat/Private/MetalDerivedData.cpp
--- a/Engine/Source/Developer/Apple/MetalShaderFormat/Private/MetalDerivedData.cpp
+++ b/Engine/Source/Developer/Apple/MetalShaderFormat/Private/MetalDerivedData.cpp
@@ -2052,7 +2052,7 @@ bool FMetalShaderOutputCooker::Build(TArray<uint8>& OutData)
 			}
 			
 			
-			if(strstr((const char*)Results.target->Data(), "spvBufferSizeConstants"))
+			if(strnstr((const char*)Results.target->Data(), "spvBufferSizeConstants", Results.target->Size()))
 			{
 				MetaData += FString::Printf(TEXT("// @SideTable: spvBufferSizeConstants(%d)\n"), SideTableIndex);
 			}
