Back out CL 11179268. WIndows doesn't know about strnstr, apparently.

diff --git a/Engine/Source/Developer/Apple/MetalShaderFormat/Private/MetalDerivedData.cpp b/Engine/Source/Developer/Apple/MetalShaderFormat/Private/MetalDerivedData.cpp
--- a/Engine/Source/Developer/Apple/MetalShaderFormat/Private/MetalDerivedData.cpp
+++ b/Engine/Source/Developer/Apple/MetalShaderFormat/Private/MetalDerivedData.cpp
@@ -2052,7 +2052,7 @@ bool FMetalShaderOutputCooker::Build(TArray<uint8>& OutData)
 			}
 			
 			
-			if(strnstr((const char*)Results.target->Data(), "spvBufferSizeConstants", Results.target->Size()))
+			if(strstr((const char*)Results.target->Data(), "spvBufferSizeConstants"))
 			{
 				MetaData += FString::Printf(TEXT("// @SideTable: spvBufferSizeConstants(%d)\n"), SideTableIndex);
 			}
