Copying Private-LoadTimes-4.24 stream to Main. Biggest changes are in Materials/Shader memory freezing.

diff --git a/Engine/Source/Developer/Android/AndroidTargetPlatform/Private/AndroidTargetPlatform.h b/Engine/Source/Developer/Android/AndroidTargetPlatform/Private/AndroidTargetPlatform.h
--- a/Engine/Source/Developer/Android/AndroidTargetPlatform/Private/AndroidTargetPlatform.h
+++ b/Engine/Source/Developer/Android/AndroidTargetPlatform/Private/AndroidTargetPlatform.h
@@ -224,8 +224,6 @@ public:
 	virtual FName GetWaveFormat( const class USoundWave* Wave ) const override;
 	virtual void GetAllWaveFormats( TArray<FName>& OutFormats) const override;
 
-	virtual FPlatformAudioCookOverrides* GetAudioCompressionSettings() const override;
-
 #endif //WITH_ENGINE
 
 	virtual bool SupportsVariants() const override;
diff --git a/Engine/Source/Developer/Apple/MetalShaderFormat/Private/MetalShaderFormat.cpp b/Engine/Source/Developer/Apple/MetalShaderFormat/Private/MetalShaderFormat.cpp
--- a/Engine/Source/Developer/Apple/MetalShaderFormat/Private/MetalShaderFormat.cpp
+++ b/Engine/Source/Developer/Apple/MetalShaderFormat/Private/MetalShaderFormat.cpp
@@ -6,7 +6,7 @@
 #include "Interfaces/IShaderFormat.h"
 #include "Interfaces/IShaderFormatModule.h"
 #include "ShaderCore.h"
-#include "Interfaces/IShaderFormatArchive.h"
+#include "ShaderCodeArchive.h"
 #include "hlslcc.h"
 #include "MetalShaderResources.h"
 #include "HAL/FileManager.h"
@@ -35,61 +35,105 @@ static FName NAME_SF_METAL_MRT_MAC(TEXT("SF_METAL_MRT_MAC"));
 static FString METAL_LIB_EXTENSION(TEXT(".metallib"));
 static FString METAL_MAP_EXTENSION(TEXT(".metalmap"));
 
-class FMetalShaderFormatArchive : public IShaderFormatArchive
+class FMetalShaderFormat : public IShaderFormat
 {
 public:
-	FMetalShaderFormatArchive(FString const& InLibraryName, FName InFormat, FString const& WorkingDirectory)
-	: LibraryName(InLibraryName)
-	, Format(InFormat)
-	, WorkingDir(WorkingDirectory)
+	enum
+	{
+		HEADER_VERSION = 68,
+	};
+	
+	struct FVersion
+	{
+		uint16 XcodeVersion;
+		uint16 HLSLCCMinor		: 8;
+		uint16 Format			: 8;
+	};
+	
+	virtual uint32 GetVersion(FName Format) const override final
+	{
+		return GetMetalFormatVersion(Format);
+	}
+	virtual void GetSupportedFormats(TArray<FName>& OutFormats) const override final
+	{
+		OutFormats.Add(NAME_SF_METAL);
+		OutFormats.Add(NAME_SF_METAL_MRT);
+		OutFormats.Add(NAME_SF_METAL_TVOS);
+		OutFormats.Add(NAME_SF_METAL_MRT_TVOS);
+		OutFormats.Add(NAME_SF_METAL_SM5_NOTESS);
+		OutFormats.Add(NAME_SF_METAL_SM5);
+		OutFormats.Add(NAME_SF_METAL_MACES3_1);
+		OutFormats.Add(NAME_SF_METAL_MRT_MAC);
+	}
+	virtual void CompileShader(FName Format, const struct FShaderCompilerInput& Input, struct FShaderCompilerOutput& Output,const FString& WorkingDirectory) const override final
 	{
-		check(LibraryName.Len() > 0);
 		check(Format == NAME_SF_METAL || Format == NAME_SF_METAL_MRT || Format == NAME_SF_METAL_TVOS || Format == NAME_SF_METAL_MRT_TVOS || Format == NAME_SF_METAL_SM5_NOTESS || Format == NAME_SF_METAL_SM5 || Format == NAME_SF_METAL_MACES3_1 || Format == NAME_SF_METAL_MRT_MAC);
-		ArchivePath = (WorkingDir / Format.GetPlainNameString());
-		IFileManager::Get().DeleteDirectory(*ArchivePath, false, true);
-		IFileManager::Get().MakeDirectory(*ArchivePath);
-		Map.Format = Format.GetPlainNameString();
+		CompileShader_Metal(Input, Output, WorkingDirectory);
 	}
-	
-	virtual FName GetFormat( void ) const
+	virtual bool CanStripShaderCode(bool const bNativeFormat) const override final
 	{
-		return Format;
+		return CanCompileBinaryShaders() && bNativeFormat;
 	}
-	
-	virtual bool AddShader( uint8 Frequency, const FSHAHash& Hash, TArray<uint8>& Code )
+	virtual bool StripShaderCode( TArray<uint8>& Code, FString const& DebugOutputDir, bool const bNative ) const override final
 	{
-		uint64 ShaderId = AppendShader_Metal(Format, ArchivePath, Hash, Code);
-		if (ShaderId)
+		return StripShader_Metal(Code, DebugOutputDir, bNative);
+    }
+	virtual bool SupportsShaderArchives() const override 
+	{ 
+		return CanCompileBinaryShaders();
+	}
+    virtual bool CreateShaderArchive(FString const& LibraryName,
+		FName Format,
+		const FString& WorkingDirectory,
+		const FString& OutputDir,
+		const FString& DebugOutputDir,
+		const FSerializedShaderArchive& InSerializedShaders,
+		const TArray<TArray<uint8>>& ShaderCode,
+		TArray<FString>* OutputFiles) const override final
+    {
+		const int32 NumShadersPerLibrary = 10000;
+		check(LibraryName.Len() > 0);
+		check(Format == NAME_SF_METAL || Format == NAME_SF_METAL_MRT || Format == NAME_SF_METAL_TVOS || Format == NAME_SF_METAL_MRT_TVOS || Format == NAME_SF_METAL_SM5_NOTESS || Format == NAME_SF_METAL_SM5 || Format == NAME_SF_METAL_MACES3_1 || Format == NAME_SF_METAL_MRT_MAC);
+
+		const FString ArchivePath = (WorkingDirectory / Format.GetPlainNameString());
+		IFileManager::Get().DeleteDirectory(*ArchivePath, false, true);
+		IFileManager::Get().MakeDirectory(*ArchivePath);
+
+		FSerializedShaderArchive SerializedShaders(InSerializedShaders);
+		check(SerializedShaders.GetNumShaders() == ShaderCode.Num());
+
+		TArray<uint8> StrippedShaderCode;
+		TArray<uint8> TempShaderCode;
+
+		TArray<TSet<uint64>> SubLibraries;
+
+		for (int32 ShaderIndex = 0; ShaderIndex < SerializedShaders.GetNumShaders(); ++ShaderIndex)
 		{
-			uint32 Index = 0;
+			SerializedShaders.DecompressShader(ShaderIndex, ShaderCode, TempShaderCode);
+			StripShader_Metal(TempShaderCode, DebugOutputDir, true);
 
-			// Add Id to our list of shaders processed successfully
-			uint32* IndexPtr = Shaders.Find(ShaderId);
-			if (IndexPtr)
-			{
-				Index = *IndexPtr;
-			}
-			else
+			uint64 ShaderId = AppendShader_Metal(Format, ArchivePath, SerializedShaders.ShaderHashes[ShaderIndex], TempShaderCode);
+			uint32 LibraryIndex = ShaderIndex / NumShadersPerLibrary;
+
+			if (ShaderId)
 			{
-				Index = (Shaders.Num() / 10000);
-				Shaders.Add(ShaderId, Index);
-				if (SubLibraries.Num() <= (int32)Index)
+				if (SubLibraries.Num() <= (int32)LibraryIndex)
 				{
 					SubLibraries.Add(TSet<uint64>());
 				}
-				SubLibraries[Index].Add(ShaderId);
+				SubLibraries[LibraryIndex].Add(ShaderId);
 			}
-			
-			// Note code copy in the map is uncompressed
-			Map.HashMap.Add(Hash, FMetalShadeEntry(Code, Index, Frequency));
+
+			FShaderCodeEntry& ShaderEntry = SerializedShaders.ShaderEntries[ShaderIndex];
+			ShaderEntry.Size = TempShaderCode.Num();
+			ShaderEntry.UncompressedSize = TempShaderCode.Num();
+
+			StrippedShaderCode.Append(TempShaderCode);
 		}
-		return (ShaderId > 0);
-	}
-	
-	virtual bool Finalize( FString OutputDir, FString DebugOutputDir, TArray<FString>* OutputFiles )
-	{
-		bool bOK = false;
 
+		SerializedShaders.Finalize();
+
+		bool bOK = false;
 		FString LibraryPlatformName = FString::Printf(TEXT("%s_%s"), *LibraryName, *Format.GetPlainNameString());
 		volatile int32 CompiledLibraries = 0;
 		TArray<FGraphEventRef> Tasks;
@@ -105,7 +149,7 @@ public:
 			}
 
 			// Enqueue the library compilation as a task so we can go wide
-			FGraphEventRef CompletionFence = FFunctionGraphTask::CreateAndDispatchWhenReady([this, LibraryPath, PartialShaders, DebugOutputDir, &CompiledLibraries]()
+			FGraphEventRef CompletionFence = FFunctionGraphTask::CreateAndDispatchWhenReady([Format, ArchivePath, LibraryPath, PartialShaders, DebugOutputDir, &CompiledLibraries]()
 			{
 				if (FinalizeLibrary_Metal(Format, ArchivePath, LibraryPath, PartialShaders, DebugOutputDir))
 				{
@@ -115,44 +159,44 @@ public:
 
 			Tasks.Add(CompletionFence);
 		}
-		
+
 #if WITH_ENGINE
-		FGraphEventRef DebugDataCompletionFence = FFunctionGraphTask::CreateAndDispatchWhenReady([this, OutputDir, LibraryPlatformName, DebugOutputDir]()
+		FGraphEventRef DebugDataCompletionFence = FFunctionGraphTask::CreateAndDispatchWhenReady([Format, OutputDir, LibraryPlatformName, DebugOutputDir]()
 		{
 			//TODO add a check in here - this will only work if we have shader archiving with debug info set.
-			
+
 			//We want to archive all the metal shader source files so that they can be unarchived into a debug location
 			//This allows the debugging of optimised metal shaders within the xcode tool set
 			//Currently using the 'tar' system tool to create a compressed tape archive
-			
+
 			//Place the archive in the same position as the .metallib file
 			FString CompressedDir = (OutputDir / TEXT("../MetaData/ShaderDebug/"));
 			IFileManager::Get().MakeDirectory(*CompressedDir, true);
-			
+
 			FString CompressedPath = (CompressedDir / LibraryPlatformName) + TEXT(".zip");
-			
+
 			IPlatformFile& PlatformFile = FPlatformFileManager::Get().GetPlatformFile();
 			IFileHandle* ZipFile = PlatformFile.OpenWrite(*CompressedPath);
 			if (ZipFile)
 			{
 				FZipArchiveWriter* ZipWriter = new FZipArchiveWriter(ZipFile);
-				
+
 				//Find the metal source files
 				TArray<FString> FilesToArchive;
-				IFileManager::Get().FindFilesRecursive( FilesToArchive, *DebugOutputDir, TEXT("*.metal"), true, false, false );
-				
+				IFileManager::Get().FindFilesRecursive(FilesToArchive, *DebugOutputDir, TEXT("*.metal"), true, false, false);
+
 				//Write the local file names into the target file
 				const FString DebugDir = DebugOutputDir / *Format.GetPlainNameString();
-				
-				for(FString FileName : FilesToArchive)
+
+				for (FString FileName : FilesToArchive)
 				{
 					TArray<uint8> FileData;
 					FFileHelper::LoadFileToArray(FileData, *FileName);
 					FPaths::MakePathRelativeTo(FileName, *DebugDir);
-					
+
 					ZipWriter->AddFile(FileName, FileData, FDateTime::Now());
 				}
-				
+
 				delete ZipWriter;
 				ZipWriter = nullptr;
 			}
@@ -162,102 +206,46 @@ public:
 			}
 		}, TStatId(), NULL, ENamedThreads::AnyThread);
 		Tasks.Add(DebugDataCompletionFence);
-#endif
+#endif // WITH_ENGINE
 
 		// Wait for tasks
 		for (auto& Task : Tasks)
 		{
 			FTaskGraphInterface::Get().WaitUntilTaskCompletes(Task);
 		}
-				
+
 		if (CompiledLibraries == SubLibraries.Num())
 		{
 			FString BinaryShaderFile = (OutputDir / LibraryPlatformName) + METAL_MAP_EXTENSION;
 			FArchive* BinaryShaderAr = IFileManager::Get().CreateFileWriter(*BinaryShaderFile);
-			if( BinaryShaderAr != NULL )
+			if (BinaryShaderAr != NULL)
 			{
-				Map.Count = SubLibraries.Num();
-				*BinaryShaderAr << Map;
+				FMetalShaderLibraryHeader Header;
+				Header.Format = Format.GetPlainNameString();
+				Header.NumLibraries = SubLibraries.Num();
+				Header.NumShadersPerLibrary = NumShadersPerLibrary;
+
+				*BinaryShaderAr << Header;
+				*BinaryShaderAr << SerializedShaders;
+				*BinaryShaderAr << StrippedShaderCode;
+
 				BinaryShaderAr->Flush();
 				delete BinaryShaderAr;
-				
+
 				if (OutputFiles)
 				{
 					OutputFiles->Add(BinaryShaderFile);
 				}
-				
+
 				bOK = true;
 			}
 		}
 
 		return bOK;
-	}
-	
-public:
-	virtual ~FMetalShaderFormatArchive() { }
-	
-private:
-	FString LibraryName;
-	FName Format;
-	FString WorkingDir;
-	FString ArchivePath;
-	TMap<uint64, uint32> Shaders;
-	TArray<TSet<uint64>> SubLibraries;
-	TSet<FString> SourceFiles;
-	FMetalShaderMap Map;
-};
 
-class FMetalShaderFormat : public IShaderFormat
-{
-public:
-	enum
-	{
-		HEADER_VERSION = 68,
-	};
-	
-	struct FVersion
-	{
-		uint16 XcodeVersion;
-		uint16 HLSLCCMinor		: 8;
-		uint16 Format			: 8;
-	};
-	
-	virtual uint32 GetVersion(FName Format) const override final
-	{
-		return GetMetalFormatVersion(Format);
-	}
-	virtual void GetSupportedFormats(TArray<FName>& OutFormats) const override final
-	{
-		OutFormats.Add(NAME_SF_METAL);
-		OutFormats.Add(NAME_SF_METAL_MRT);
-		OutFormats.Add(NAME_SF_METAL_TVOS);
-		OutFormats.Add(NAME_SF_METAL_MRT_TVOS);
-		OutFormats.Add(NAME_SF_METAL_SM5_NOTESS);
-		OutFormats.Add(NAME_SF_METAL_SM5);
-		OutFormats.Add(NAME_SF_METAL_MACES3_1);
-		OutFormats.Add(NAME_SF_METAL_MRT_MAC);
-	}
-	virtual void CompileShader(FName Format, const struct FShaderCompilerInput& Input, struct FShaderCompilerOutput& Output,const FString& WorkingDirectory) const override final
-	{
-		check(Format == NAME_SF_METAL || Format == NAME_SF_METAL_MRT || Format == NAME_SF_METAL_TVOS || Format == NAME_SF_METAL_MRT_TVOS || Format == NAME_SF_METAL_SM5_NOTESS || Format == NAME_SF_METAL_SM5 || Format == NAME_SF_METAL_MACES3_1 || Format == NAME_SF_METAL_MRT_MAC);
-		CompileShader_Metal(Input, Output, WorkingDirectory);
-	}
-	virtual bool CanStripShaderCode(bool const bNativeFormat) const override final
-	{
-		return CanCompileBinaryShaders() && bNativeFormat;
-	}
-	virtual bool StripShaderCode( TArray<uint8>& Code, FString const& DebugOutputDir, bool const bNative ) const override final
-	{
-		return StripShader_Metal(Code, DebugOutputDir, bNative);
-    }
-	virtual bool SupportsShaderArchives() const override 
-	{ 
-		return CanCompileBinaryShaders();
-	}
-    virtual class IShaderFormatArchive* CreateShaderArchive( FString const& LibraryName, FName Format, const FString& WorkingDirectory ) const override final
-    {
-		return new FMetalShaderFormatArchive(LibraryName, Format, WorkingDirectory);
+		//Map.Format = Format.GetPlainNameString();
     }
+
 	virtual bool CanCompileBinaryShaders() const override final
 	{
 #if PLATFORM_MAC
diff --git a/Engine/Source/Developer/BlueprintCompilerCppBackend/Private/BlueprintCompilerCppBackendUtils.cpp b/Engine/Source/Developer/BlueprintCompilerCppBackend/Private/BlueprintCompilerCppBackendUtils.cpp
--- a/Engine/Source/Developer/BlueprintCompilerCppBackend/Private/BlueprintCompilerCppBackendUtils.cpp
+++ b/Engine/Source/Developer/BlueprintCompilerCppBackend/Private/BlueprintCompilerCppBackendUtils.cpp
@@ -1102,8 +1102,8 @@ FString FEmitHelper::LiteralTerm(FEmitterLocalContext& EmitterContext, const FLi
 				UE_LOG(LogK2Compiler, Error, TEXT("FEmitHelper::LiteralTerm cannot parse array value \"%s\" error: %s class: %s"), *CustomValue, *ImportError, *GetPathNameSafe(EmitterContext.GetCurrentlyGeneratedClass()));
 			}
 
-			const int32 NumElements = ScriptArray.Num();
 			FScriptArrayHelper ScriptArrayHelper(ArrayProperty, &ScriptArray);
+			const int32 NumElements = ScriptArrayHelper.Num();
 
 			FLiteralTermParams InnerTermParams;
 			Schema->ConvertPropertyToPinType(ArrayProperty->Inner, InnerTermParams.Type);
diff --git a/Engine/Source/Developer/GameplayDebugger/Private/GameplayDebuggerRenderingComponent.cpp b/Engine/Source/Developer/GameplayDebugger/Private/GameplayDebuggerRenderingComponent.cpp
--- a/Engine/Source/Developer/GameplayDebugger/Private/GameplayDebuggerRenderingComponent.cpp
+++ b/Engine/Source/Developer/GameplayDebugger/Private/GameplayDebuggerRenderingComponent.cpp
@@ -174,9 +174,9 @@ FBoxSphereBounds UGameplayDebuggerRenderingComponent::CalcBounds(const FTransfor
 	return FBoxSphereBounds(FBox::BuildAABB(FVector::ZeroVector, FVector(1000000.0f, 1000000.0f, 1000000.0f)));
 }
 
-void UGameplayDebuggerRenderingComponent::CreateRenderState_Concurrent()
+void UGameplayDebuggerRenderingComponent::CreateRenderState_Concurrent(FRegisterComponentContext* Context)
 {
-	Super::CreateRenderState_Concurrent();
+	Super::CreateRenderState_Concurrent(Context);
 
 	GameplayDebuggerDebugDrawDelegateHelper.RegisterDebugDrawDelgate();
 }
diff --git a/Engine/Source/Developer/IOS/IOSTargetPlatform/Private/IOSTargetPlatform.cpp b/Engine/Source/Developer/IOS/IOSTargetPlatform/Private/IOSTargetPlatform.cpp
--- a/Engine/Source/Developer/IOS/IOSTargetPlatform/Private/IOSTargetPlatform.cpp
+++ b/Engine/Source/Developer/IOS/IOSTargetPlatform/Private/IOSTargetPlatform.cpp
@@ -724,140 +724,5 @@ void FIOSTargetPlatform::GetAllWaveFormats(TArray<FName>& OutFormat) const
 	OutFormat.Add(NAME_ADPCM);
 }
 
-namespace IOS
-{
-	void CachePlatformAudioCookOverrides(FPlatformAudioCookOverrides& OutOverrides)
-	{
-		const TCHAR* CategoryName = TEXT("/Script/IOSRuntimeSettings.IOSRuntimeSettings");
-
-		int32 SoundCueQualityIndex = INDEX_NONE;
-		if (GConfig->GetInt(CategoryName, TEXT("SoundCueCookQualityIndex"), SoundCueQualityIndex, GEngineIni))
-		{
-			OutOverrides.SoundCueCookQualityIndex = SoundCueQualityIndex;
-		}
-
-		GConfig->GetBool(CategoryName, TEXT("bUseAudioStreamCaching"), OutOverrides.bUseStreamCaching, GEngineIni);
-
-		/** Memory Load On Demand Settings */
-		if (OutOverrides.bUseStreamCaching)
-		{
-			// Cache size:
-			int32 RetrievedCacheSize = 32 * 1024;
-			GConfig->GetInt(CategoryName, TEXT("CacheSizeKB"), RetrievedCacheSize, GEngineIni);
-			OutOverrides.StreamCachingSettings.CacheSizeKB = RetrievedCacheSize;
-		}
-
-		GConfig->GetBool(CategoryName, TEXT("bResampleForDevice"), OutOverrides.bResampleForDevice, GEngineIni);
-
-		GConfig->GetFloat(CategoryName, TEXT("CompressionQualityModifier"), OutOverrides.CompressionQualityModifier, GEngineIni);
-
-		GConfig->GetFloat(CategoryName, TEXT("AutoStreamingThreshold"), OutOverrides.AutoStreamingThreshold, GEngineIni);
-
-		//Cache sample rate map:
-		float RetrievedSampleRate = -1.0f;
-
-		GConfig->GetFloat(CategoryName, TEXT("MaxSampleRate"), RetrievedSampleRate, GEngineIni);
-		float* FoundSampleRate = OutOverrides.PlatformSampleRates.Find(ESoundwaveSampleRateSettings::Max);
-
-		if (FoundSampleRate)
-		{
-			if (!FMath::IsNearlyEqual(*FoundSampleRate, RetrievedSampleRate))
-			{
-				*FoundSampleRate = RetrievedSampleRate;
-			}
-
-		}
-		else
-		{
-			OutOverrides.PlatformSampleRates.Add(ESoundwaveSampleRateSettings::Max, RetrievedSampleRate);
-		}
-
-		RetrievedSampleRate = -1.0f;
-
-		GConfig->GetFloat(CategoryName, TEXT("HighSampleRate"), RetrievedSampleRate, GEngineIni);
-		FoundSampleRate = OutOverrides.PlatformSampleRates.Find(ESoundwaveSampleRateSettings::High);
-
-		if (FoundSampleRate)
-		{
-			if (!FMath::IsNearlyEqual(*FoundSampleRate, RetrievedSampleRate))
-			{
-				*FoundSampleRate = RetrievedSampleRate;
-			}
-
-		}
-		else
-		{
-			OutOverrides.PlatformSampleRates.Add(ESoundwaveSampleRateSettings::High, RetrievedSampleRate);
-		}
-
-
-		RetrievedSampleRate = -1.0f;
-
-		GConfig->GetFloat(CategoryName, TEXT("MedSampleRate"), RetrievedSampleRate, GEngineIni);
-		FoundSampleRate = OutOverrides.PlatformSampleRates.Find(ESoundwaveSampleRateSettings::Medium);
-
-		if (FoundSampleRate)
-		{
-			if (!FMath::IsNearlyEqual(*FoundSampleRate, RetrievedSampleRate))
-			{
-				*FoundSampleRate = RetrievedSampleRate;
-			}
-		}
-		else
-		{
-			OutOverrides.PlatformSampleRates.Add(ESoundwaveSampleRateSettings::Medium, RetrievedSampleRate);
-		}
-
-		RetrievedSampleRate = -1.0f;
-
-		GConfig->GetFloat(CategoryName, TEXT("LowSampleRate"), RetrievedSampleRate, GEngineIni);
-		FoundSampleRate = OutOverrides.PlatformSampleRates.Find(ESoundwaveSampleRateSettings::Low);
-
-		if (FoundSampleRate)
-		{
-			if (!FMath::IsNearlyEqual(*FoundSampleRate, RetrievedSampleRate))
-			{
-				*FoundSampleRate = RetrievedSampleRate;
-			}
-		}
-		else
-		{
-			OutOverrides.PlatformSampleRates.Add(ESoundwaveSampleRateSettings::Low, RetrievedSampleRate);
-		}
-
-		RetrievedSampleRate = -1.0f;
-
-		GConfig->GetFloat(CategoryName, TEXT("MinSampleRate"), RetrievedSampleRate, GEngineIni);
-		FoundSampleRate = OutOverrides.PlatformSampleRates.Find(ESoundwaveSampleRateSettings::Min);
-
-		if (FoundSampleRate)
-		{
-			if (!FMath::IsNearlyEqual(*FoundSampleRate, RetrievedSampleRate))
-			{
-				*FoundSampleRate = RetrievedSampleRate;
-			}
-		}
-		else
-		{
-			OutOverrides.PlatformSampleRates.Add(ESoundwaveSampleRateSettings::Min, RetrievedSampleRate);
-		}
-	}
-}
-
-FPlatformAudioCookOverrides* FIOSTargetPlatform::GetAudioCompressionSettings() const
-{
-	static FPlatformAudioCookOverrides Settings;
-
-	static bool bCachedPlatformSettings = false;
-
-	if (!bCachedPlatformSettings)
-	{
-		IOS::CachePlatformAudioCookOverrides(Settings);
-		bCachedPlatformSettings = true;
-	}
-
-	return &Settings;
-}
-
 #endif // WITH_ENGINE
 
diff --git a/Engine/Source/Developer/IOS/IOSTargetPlatform/Private/IOSTargetPlatform.h b/Engine/Source/Developer/IOS/IOSTargetPlatform/Private/IOSTargetPlatform.h
--- a/Engine/Source/Developer/IOS/IOSTargetPlatform/Private/IOSTargetPlatform.h
+++ b/Engine/Source/Developer/IOS/IOSTargetPlatform/Private/IOSTargetPlatform.h
@@ -133,8 +133,6 @@ public:
 	virtual FName GetWaveFormat( const class USoundWave* Wave ) const override;
 	virtual void GetAllWaveFormats(TArray<FName>& OutFormat) const override;
 
-	virtual FPlatformAudioCookOverrides* GetAudioCompressionSettings() const override;
-
 #endif // WITH_ENGINE
 
 	virtual void GetBuildProjectSettingKeys(FString& OutSection, TArray<FString>& InBoolKeys, TArray<FString>& InIntKeys, TArray<FString>& InStringKeys) const override
diff --git a/Engine/Source/Developer/Linux/LinuxTargetPlatform/Private/LinuxTargetPlatform.h b/Engine/Source/Developer/Linux/LinuxTargetPlatform/Private/LinuxTargetPlatform.h
--- a/Engine/Source/Developer/Linux/LinuxTargetPlatform/Private/LinuxTargetPlatform.h
+++ b/Engine/Source/Developer/Linux/LinuxTargetPlatform/Private/LinuxTargetPlatform.h
@@ -321,7 +321,7 @@ public:
 			return NAME_ADPCM;
 		}
 
-		if (Wave->IsStreaming())
+		if (Wave->IsStreaming(*this->IniPlatformName()))
 		{
 			return NAME_OPUS;
 		}
@@ -340,11 +340,6 @@ public:
 		OutFormats.Add(NAME_OPUS);
 	}
 
-	virtual FPlatformAudioCookOverrides* GetAudioCompressionSettings() const override
-	{
-		return nullptr;
-	}
-
 #endif //WITH_ENGINE
 
 	virtual bool SupportsVariants() const override
diff --git a/Engine/Source/Developer/LogVisualizer/Private/VisualLoggerRenderingActor.cpp b/Engine/Source/Developer/LogVisualizer/Private/VisualLoggerRenderingActor.cpp
--- a/Engine/Source/Developer/LogVisualizer/Private/VisualLoggerRenderingActor.cpp
+++ b/Engine/Source/Developer/LogVisualizer/Private/VisualLoggerRenderingActor.cpp
@@ -114,9 +114,9 @@ FBoxSphereBounds UVisualLoggerRenderingComponent::CalcBounds(const FTransform& L
 	return MyBounds;
 }
 
-void UVisualLoggerRenderingComponent::CreateRenderState_Concurrent()
+void UVisualLoggerRenderingComponent::CreateRenderState_Concurrent(FRegisterComponentContext* Context)
 {
-	Super::CreateRenderState_Concurrent();
+	Super::CreateRenderState_Concurrent(Context);
 
 #if WITH_EDITOR
 	DebugDrawDelegateHelper.RegisterDebugDrawDelgate();
diff --git a/Engine/Source/Developer/Lumin/LuminTargetPlatform/Private/LuminTargetPlatform.h b/Engine/Source/Developer/Lumin/LuminTargetPlatform/Private/LuminTargetPlatform.h
--- a/Engine/Source/Developer/Lumin/LuminTargetPlatform/Private/LuminTargetPlatform.h
+++ b/Engine/Source/Developer/Lumin/LuminTargetPlatform/Private/LuminTargetPlatform.h
@@ -76,11 +76,6 @@ public:
 		OutFormats.Add(FName(TEXT("FullHDR")));
 	}
 
-	virtual FPlatformAudioCookOverrides* GetAudioCompressionSettings() const override
-	{
-		return nullptr;
-	}
-	
 	// True if the project requires encoded HDR reflection captures
 	bool bRequiresEncodedHDRReflectionCaptures;
 #endif //WITH_ENGINE
diff --git a/Engine/Source/Developer/Mac/MacTargetPlatform/Private/GenericMacTargetPlatform.h b/Engine/Source/Developer/Mac/MacTargetPlatform/Private/GenericMacTargetPlatform.h
--- a/Engine/Source/Developer/Mac/MacTargetPlatform/Private/GenericMacTargetPlatform.h
+++ b/Engine/Source/Developer/Mac/MacTargetPlatform/Private/GenericMacTargetPlatform.h
@@ -320,7 +320,7 @@ return TSuper::SupportsFeature(Feature);
 			return NAME_ADPCM;
 		}
 
-		if (Wave->IsStreaming())
+		if (Wave->IsStreaming(*this->IniPlatformName()))
 		{
 			return NAME_OPUS;
 		}
@@ -339,11 +339,6 @@ return TSuper::SupportsFeature(Feature);
 		OutFormats.Add(NAME_OPUS);
 	}
 
-	virtual FPlatformAudioCookOverrides* GetAudioCompressionSettings() const override
-	{
-		return nullptr;
-	}
-
 #endif //WITH_ENGINE
 
 
diff --git a/Engine/Source/Developer/MaterialBaking/Private/ExportMaterialProxy.h b/Engine/Source/Developer/MaterialBaking/Private/ExportMaterialProxy.h
--- a/Engine/Source/Developer/MaterialBaking/Private/ExportMaterialProxy.h
+++ b/Engine/Source/Developer/MaterialBaking/Private/ExportMaterialProxy.h
@@ -229,7 +229,7 @@ public:
 	{
 		SetQualityLevelProperties(EMaterialQualityLevel::High, false, GMaxRHIFeatureLevel);
 		Material = InMaterialInterface->GetMaterial();
-		InMaterialInterface->AppendReferencedTextures(ReferencedTextures);
+		ReferencedTextures = InMaterialInterface->GetReferencedTextures();
 
 		const FMaterialResource* Resource = InMaterialInterface->GetMaterialResource(GMaxRHIFeatureLevel);
 
@@ -240,8 +240,6 @@ public:
 		// to be able to set back GameThreadShaderMap after async compilation.
 		Id = ResourceId.BaseMaterialId;
 
-		FStaticParameterSet StaticParamSet;
-		Resource->GetStaticParameterSet(GMaxRHIShaderPlatform, StaticParamSet);
 
 		{
 			TArray<FShaderType*> ShaderTypes;
@@ -272,7 +270,7 @@ public:
 			break;
 		};
 
-		CacheShaders(ResourceId, &StaticParamSet, GMaxRHIShaderPlatform);
+		CacheShaders(ResourceId, GMaxRHIShaderPlatform);
 	}
 
 	/** This override is required otherwise the shaders aren't ready for use when the surface is rendered resulting in a blank image */
@@ -296,7 +294,7 @@ public:
 		return bCorrectVertexFactory && bPCPlatform && bCorrectFrequency;
 	}
 
-	virtual const TArray<UObject*>& GetReferencedTextures() const override
+	virtual TArrayView<UObject* const> GetReferencedTextures() const override
 	{
 		return ReferencedTextures;
 	}
@@ -316,22 +314,22 @@ public:
 		}
 	}
 
-	virtual bool GetVectorValue(const FMaterialParameterInfo& ParameterInfo, FLinearColor* OutValue, const FMaterialRenderContext& Context) const override
+	virtual bool GetVectorValue(const FHashedMaterialParameterInfo& ParameterInfo, FLinearColor* OutValue, const FMaterialRenderContext& Context) const override
 	{
 		return MaterialInterface->GetRenderProxy()->GetVectorValue(ParameterInfo, OutValue, Context);
 	}
 
-	virtual bool GetScalarValue(const FMaterialParameterInfo& ParameterInfo, float* OutValue, const FMaterialRenderContext& Context) const override
+	virtual bool GetScalarValue(const FHashedMaterialParameterInfo& ParameterInfo, float* OutValue, const FMaterialRenderContext& Context) const override
 	{
 		return MaterialInterface->GetRenderProxy()->GetScalarValue(ParameterInfo, OutValue, Context);
 	}
 
-	virtual bool GetTextureValue(const FMaterialParameterInfo& ParameterInfo, const UTexture** OutValue, const FMaterialRenderContext& Context) const override
+	virtual bool GetTextureValue(const FHashedMaterialParameterInfo& ParameterInfo, const UTexture** OutValue, const FMaterialRenderContext& Context) const override
 	{
 		return MaterialInterface->GetRenderProxy()->GetTextureValue(ParameterInfo, OutValue, Context);
 	}
 
-	virtual bool GetTextureValue(const FMaterialParameterInfo& ParameterInfo, const URuntimeVirtualTexture** OutValue, const FMaterialRenderContext& Context) const override
+	virtual bool GetTextureValue(const FHashedMaterialParameterInfo& ParameterInfo, const URuntimeVirtualTexture** OutValue, const FMaterialRenderContext& Context) const override
 	{
 		return MaterialInterface->GetRenderProxy()->GetTextureValue(ParameterInfo, OutValue, Context);
 	}
diff --git a/Engine/Source/Developer/MaterialUtilities/Private/MaterialUtilities.cpp b/Engine/Source/Developer/MaterialUtilities/Private/MaterialUtilities.cpp
--- a/Engine/Source/Developer/MaterialUtilities/Private/MaterialUtilities.cpp
+++ b/Engine/Source/Developer/MaterialUtilities/Private/MaterialUtilities.cpp
@@ -410,7 +410,7 @@ public:
 	{
 		SetQualityLevelProperties(EMaterialQualityLevel::High, false, GMaxRHIFeatureLevel);
 		Material = InMaterialInterface->GetMaterial();
-		InMaterialInterface->AppendReferencedTextures(ReferencedTextures);
+		ReferencedTextures = InMaterialInterface->GetReferencedTextures();
 		FPlatformMisc::CreateGuid(Id);
 
 		FMaterialResource* Resource = InMaterialInterface->GetMaterialResource(GMaxRHIFeatureLevel);
@@ -418,9 +418,6 @@ public:
 		FMaterialShaderMapId ResourceId;
 		Resource->GetShaderMapId(GMaxRHIShaderPlatform, ResourceId);
 
-		FStaticParameterSet StaticParamSet;
-		Resource->GetStaticParameterSet(GMaxRHIShaderPlatform, StaticParamSet);
-
 		{
 			TArray<FShaderType*> ShaderTypes;
 			TArray<FVertexFactoryType*> VFTypes;
@@ -450,7 +447,7 @@ public:
 			break;
 		};
 		
-		CacheShaders(ResourceId, &StaticParamSet, GMaxRHIShaderPlatform);
+		CacheShaders(ResourceId, GMaxRHIShaderPlatform);
 	}
 
 	virtual bool IsUsedWithStaticLighting() const { return true; }
@@ -474,7 +471,7 @@ public:
 		return true;
 	}
 
-	virtual const TArray<UObject*>& GetReferencedTextures() const override
+	virtual TArrayView<UObject* const> GetReferencedTextures() const override
 	{
 		return ReferencedTextures;
 	}
@@ -494,22 +491,22 @@ public:
 		}
 	}
 
-	virtual bool GetVectorValue(const FMaterialParameterInfo& ParameterInfo, FLinearColor* OutValue, const FMaterialRenderContext& Context) const override
+	virtual bool GetVectorValue(const FHashedMaterialParameterInfo& ParameterInfo, FLinearColor* OutValue, const FMaterialRenderContext& Context) const override
 	{
 		return MaterialInterface->GetRenderProxy()->GetVectorValue(ParameterInfo, OutValue, Context);
 	}
 
-	virtual bool GetScalarValue(const FMaterialParameterInfo& ParameterInfo, float* OutValue, const FMaterialRenderContext& Context) const override
+	virtual bool GetScalarValue(const FHashedMaterialParameterInfo& ParameterInfo, float* OutValue, const FMaterialRenderContext& Context) const override
 	{
 		return MaterialInterface->GetRenderProxy()->GetScalarValue(ParameterInfo, OutValue, Context);
 	}
 
-	virtual bool GetTextureValue(const FMaterialParameterInfo& ParameterInfo,const UTexture** OutValue, const FMaterialRenderContext& Context) const override
+	virtual bool GetTextureValue(const FHashedMaterialParameterInfo& ParameterInfo,const UTexture** OutValue, const FMaterialRenderContext& Context) const override
 	{
 		return MaterialInterface->GetRenderProxy()->GetTextureValue(ParameterInfo,OutValue,Context);
 	}
 
-	virtual bool GetTextureValue(const FMaterialParameterInfo& ParameterInfo, const URuntimeVirtualTexture** OutValue, const FMaterialRenderContext& Context) const override
+	virtual bool GetTextureValue(const FHashedMaterialParameterInfo& ParameterInfo, const URuntimeVirtualTexture** OutValue, const FMaterialRenderContext& Context) const override
 	{
 		return MaterialInterface->GetRenderProxy()->GetTextureValue(ParameterInfo, OutValue, Context);
 	}
@@ -2836,8 +2833,7 @@ void FMaterialUtilities::DetermineMaterialImportance(const TArray<UMaterialInter
 	int32 SummedSize = 0;
 	for (UMaterialInterface* Material : InMaterials)
 	{
-		TArray<UObject*> UsedTextures;
-		Material->AppendReferencedTextures(UsedTextures);
+		TArray<UObject*> UsedTextures(Material->GetReferencedTextures());
 		if (UMaterialInstance* MaterialInstance = Cast<UMaterialInstance>(Material))
 		{
 			for (const FTextureParameterValue& TextureParameter : MaterialInstance->TextureParameterValues)
diff --git a/Engine/Source/Developer/MeshBuilder/Private/BuildOptimizationHelper.h b/Engine/Source/Developer/MeshBuilder/Private/BuildOptimizationHelper.h
--- a/Engine/Source/Developer/MeshBuilder/Private/BuildOptimizationHelper.h
+++ b/Engine/Source/Developer/MeshBuilder/Private/BuildOptimizationHelper.h
@@ -223,7 +223,7 @@ public:
 		TArray<uint32>& OutDepthIndices,
 		const TArray<FStaticMeshBuildVertex>& InVertices,
 		const TArray<uint32>& InIndices,
-		const TArray<FStaticMeshSection>& InSections
+		const TArrayView<FStaticMeshSection>& InSections
 	)
 	{
 		int32 NumVertices = InVertices.Num();
diff --git a/Engine/Source/Developer/MeshBuilder/Private/StaticMeshBuilder.cpp b/Engine/Source/Developer/MeshBuilder/Private/StaticMeshBuilder.cpp
--- a/Engine/Source/Developer/MeshBuilder/Private/StaticMeshBuilder.cpp
+++ b/Engine/Source/Developer/MeshBuilder/Private/StaticMeshBuilder.cpp
@@ -262,7 +262,9 @@ bool FStaticMeshBuilder::Build(FStaticMeshRenderData& StaticMeshRenderData, USta
 
 		StaticMeshLOD.Sections.Empty(PolygonGroups.Num());
 		TArray<int32> RemapVerts; //Because we will remove MeshVertex that are redundant, we need a remap
-		TArray<int32> &WedgeMap = StaticMeshLOD.WedgeMap;
+								  //Render data Wedge map is only set for LOD 0???
+
+		TArray<int32>& WedgeMap = StaticMeshLOD.WedgeMap;
 		WedgeMap.Reset();
 
 		//Prepare the PerSectionIndices array so we can optimize the index buffer for the GPU
diff --git a/Engine/Source/Developer/MeshUtilities/Private/MeshDistanceFieldUtilities.cpp b/Engine/Source/Developer/MeshUtilities/Private/MeshDistanceFieldUtilities.cpp
--- a/Engine/Source/Developer/MeshUtilities/Private/MeshDistanceFieldUtilities.cpp
+++ b/Engine/Source/Developer/MeshUtilities/Private/MeshDistanceFieldUtilities.cpp
@@ -681,7 +681,7 @@ void FMeshUtilities::GenerateSignedDistanceFieldVolumeData(
 					int32 CompressedSize = TempCompressedMemory.Num() * TempCompressedMemory.GetTypeSize();
 
 					verify(FCompression::CompressMemory(
-						NAME_Zlib, 
+						NAME_LZ4, 
 						TempCompressedMemory.GetData(), 
 						CompressedSize, 
 						QuantizedDistanceFieldVolume.GetData(), 
@@ -775,7 +775,7 @@ void FMeshUtilities::DownSampleDistanceFieldVolumeData(FDistanceFieldVolumeData&
 	if (bDataIsCompressed)
 	{
 		UncompressedData.AddUninitialized(UncompressedSize);
-		verify(FCompression::UncompressMemory(NAME_Zlib, UncompressedData.GetData(), UncompressedSize, DistanceFieldData.CompressedDistanceFieldVolume.GetData(), DistanceFieldData.CompressedDistanceFieldVolume.Num()));
+		verify(FCompression::UncompressMemory(NAME_LZ4, UncompressedData.GetData(), UncompressedSize, DistanceFieldData.CompressedDistanceFieldVolume.GetData(), DistanceFieldData.CompressedDistanceFieldVolume.Num()));
 		SourceData = &UncompressedData;
 	}
 
@@ -864,7 +864,7 @@ void FMeshUtilities::DownSampleDistanceFieldVolumeData(FDistanceFieldVolumeData&
 		int32 CompressedSize = TempCompressedMemory.Num() * TempCompressedMemory.GetTypeSize();
 
 		verify(FCompression::CompressMemory(
-			NAME_Zlib,
+			NAME_LZ4,
 			TempCompressedMemory.GetData(),
 			CompressedSize,
 			DownSampledTexture.GetData(),
diff --git a/Engine/Source/Developer/MeshUtilities/Private/MeshUtilities.cpp b/Engine/Source/Developer/MeshUtilities/Private/MeshUtilities.cpp
--- a/Engine/Source/Developer/MeshUtilities/Private/MeshUtilities.cpp
+++ b/Engine/Source/Developer/MeshUtilities/Private/MeshUtilities.cpp
@@ -2121,7 +2121,7 @@ static void BuildDepthOnlyIndexBuffer(
 	TArray<uint32>& OutDepthIndices,
 	const TArray<FStaticMeshBuildVertex>& InVertices,
 	const TArray<uint32>& InIndices,
-	const TArray<FStaticMeshSection>& InSections
+	const TArrayView<FStaticMeshSection>& InSections
 	)
 {
 	int32 NumVertices = InVertices.Num();
@@ -2771,6 +2771,7 @@ public:
 			// Build and cache optimize vertex and index buffers.
 			{
 				// TODO_STATICMESH: The wedge map is only valid for LODIndex 0 if no reduction has been performed.
+				// TODO - write directly to TMemoryImageArray
 				// We can compute an approximate one instead for other LODs.
 				TArray<int32> TempWedgeMap;
 				TArray<int32>& WedgeMap = InOutModels[LODIndex].ReductionSettings.PercentTriangles >= 1.0f ? LODModel.WedgeMap : TempWedgeMap;
diff --git a/Engine/Source/Developer/PakFileUtilities/Private/PakFileUtilities.cpp b/Engine/Source/Developer/PakFileUtilities/Private/PakFileUtilities.cpp
--- a/Engine/Source/Developer/PakFileUtilities/Private/PakFileUtilities.cpp
+++ b/Engine/Source/Developer/PakFileUtilities/Private/PakFileUtilities.cpp
@@ -23,6 +23,7 @@
 #include "DerivedDataCacheInterface.h"
 #include "Serialization/MemoryReader.h"
 #include "Serialization/MemoryWriter.h"
+#include "Misc/DataDrivenPlatformInfoRegistry.h"
 
 IMPLEMENT_MODULE(FDefaultModuleImpl, PakFileUtilities);
 
@@ -483,6 +484,7 @@ struct FPakCommandLineParameters
 		, bAsyncCompression(false)
 		, bAlignFilesLargerThanBlock(false)
 		, bForceCompress(false)
+		, bAllowForIndexUnload(false)
 	{
 	}
 
@@ -506,6 +508,8 @@ struct FPakCommandLineParameters
 	bool bAsyncCompression;
 	bool bAlignFilesLargerThanBlock;	// Align files that are larger than block size
 	bool bForceCompress; // Force all files that request compression to be compressed, even if that results in a larger file size
+	bool bAllowForIndexUnload;
+	FString DataDrivenPlatformName;
 };
 
 struct FPakEntryPair
@@ -1113,8 +1117,10 @@ void ProcessCommandLine(const TCHAR* CmdLine, const TArray<FString>& NonOptionAr
 
 	if (FParse::Value(CmdLine, TEXT("-create="), ResponseFile))
 	{
+		CmdLineParameters.bAllowForIndexUnload = FParse::Param(CmdLine, TEXT("allowforindexunload"));
 		CmdLineParameters.GeneratePatch = FParse::Value(CmdLine, TEXT("-generatepatch="), CmdLineParameters.SourcePatchPakFilename);
 		FParse::Value(CmdLine, TEXT("-outputchangedfiles="), CmdLineParameters.ChangedFilesOutputFilename);
+		FParse::Value(CmdLine, TEXT("-platform="), CmdLineParameters.DataDrivenPlatformName);
 
 		bool bCompress = FParse::Param(CmdLine, TEXT("compress"));
 		bool bEncrypt = FParse::Param(CmdLine, TEXT("encrypt"));
@@ -2344,60 +2350,166 @@ bool CreatePakFile(const TCHAR* Filename, TArray<FPakInputPair>& FilesToAdd, con
 	// Remember IndexOffset
 	Info.IndexOffset = PakFileHandle->Tell();
 
-	// Serialize Pak Index at the end of Pak File
-	TArray<uint8> IndexData;
-	FMemoryWriter IndexWriter(IndexData);
-	IndexWriter.SetByteSwapping(PakFileHandle->ForceByteSwapping());
-	int32 NumEntries = Index.Num();
-	IndexWriter << MountPoint;
-	IndexWriter << NumEntries;
+	FPakFileData Data;
+	FPakFile::MakeDirectoryFromPath(MountPoint);
+	Data.MountPoint = MountPoint;
+
+	// process each entry similar to how loading was working
 	for (int32 EntryIndex = 0; EntryIndex < Index.Num(); EntryIndex++)
 	{
 		FPakEntryPair& Entry = Index[EntryIndex];
-		IndexWriter << Entry.Filename;
-		Entry.Info.Serialize(IndexWriter, Info.Version);
+		// add raw entry
+		Data.Files.Add(Entry.Info);
 
-		if (RequiredPatchPadding > 0)
+		// make index entry
+		FString Path = FPaths::GetPath(Entry.Filename);
+		FPakFile::MakeDirectoryFromPath(Path);
+		FPakDirectory* Directory = Data.Index.Find(Path);
+		if (Directory != nullptr)
+		{
+			Directory->Add(FPaths::GetCleanFilename(Entry.Filename), EntryIndex);
+		}
+		else
 		{
-			int64 EntrySize = Entry.Info.GetSerializedSize(FPakInfo::PakFile_Version_Latest);
-			int64 TotalSizeToWrite = Entry.Info.Size + EntrySize;
-			if (TotalSizeToWrite >= RequiredPatchPadding)
+			FPakDirectory& NewDirectory = Data.Index.Add(Path);
+			NewDirectory.Add(FPaths::GetCleanFilename(Entry.Filename), EntryIndex);
+
+			// add the parent directories up to the mount point
+			while (MountPoint != Path)
 			{
-				int64 RealStart = Entry.Info.Offset;
-				if ((RealStart % RequiredPatchPadding) != 0 &&
-					!Entry.Filename.EndsWith(TEXT("uexp")) && // these are export sections of larger files and may be packed with uasset/umap and so we don't need a warning here
-					!(Entry.Filename.EndsWith(TEXT(".m.ubulk")) && CmdLineParameters.AlignForMemoryMapping > 0)) // Bulk padding unaligns patch padding and so we don't need a warning here
+				Path = Path.Left(Path.Len() - 1);
+				int32 Offset = 0;
+				if (Path.FindLastChar('/', Offset))
 				{
-					UE_LOG(LogPakFile, Warning, TEXT("File at offset %lld of size %lld not aligned to patch size %i"), RealStart, Entry.Info.Size, RequiredPatchPadding);
+					Path = Path.Left(Offset);
+					FPakFile::MakeDirectoryFromPath(Path);
+					if (Data.Index.Find(Path) == nullptr)
+					{
+						Data.Index.Add(Path);
+					}
+				}
+				else
+				{
+					Path = MountPoint;
 				}
 			}
 		}
 	}
 
-	if (Info.bEncryptedIndex)
+	// look up platform info via DataDrivenPlatformInfo.ini (null here means that there was no platform specified, and we cannot freeze)
+	const FDataDrivenPlatformInfoRegistry::FPlatformInfo* PlatformInfo = nullptr;
+	if (CmdLineParameters.DataDrivenPlatformName.Len() > 0)
 	{
-		int32 OriginalSize = IndexData.Num();
-		int32 AlignedSize = Align(OriginalSize, FAES::AESBlockSize);
-
-		for (int32 PaddingIndex = IndexData.Num(); PaddingIndex < AlignedSize; ++PaddingIndex)
+		const TMap<FString, FDataDrivenPlatformInfoRegistry::FPlatformInfo>& Infos = FDataDrivenPlatformInfoRegistry::GetAllPlatformInfos();
+		PlatformInfo = Infos.Find(CmdLineParameters.DataDrivenPlatformName);
+		if (PlatformInfo == nullptr)
 		{
-			uint8 Byte = IndexData[PaddingIndex % OriginalSize];
-			IndexData.Add(Byte);
+			UE_LOG(LogPakFile, Fatal, TEXT("Unable to find DataDrivenPlatform '%s', unable to continue safely"), *CmdLineParameters.DataDrivenPlatformName);
 		}
 	}
 
-	FSHA1::HashBuffer(IndexData.GetData(), IndexData.Num(), Info.IndexHash.Hash);
+	// check to make sure freezing of the index is okay (target specified, it's 64-bit, and we don't want to modify it at runtime)
+	bool bAllowFreezing = PlatformInfo != nullptr && PlatformInfo->Freezing_b32Bit == false && CmdLineParameters.bAllowForIndexUnload == false;
 
-	if (Info.bEncryptedIndex)
+	// Disable encrypted freezing until frozen decryption is implemented in FPakFile::LoadIndex()
+	bAllowFreezing = bAllowFreezing && !Info.bEncryptedIndex;
+
+	// if the want to allow for unloading of indices, then 
+	if (bAllowFreezing)
 	{
-		check(InKeyChain.MasterEncryptionKey);
-		FAES::EncryptData(IndexData.GetData(), IndexData.Num(), InKeyChain.MasterEncryptionKey->Key);
-		TotalEncryptedDataSize += IndexData.Num();
+		UE_LOG(LogPakFile, Display, TEXT("ENABLING pak file index freezing"));
+
+		FMemoryImage MemoryImage;
+		// we don't have any editor only data while freezing, so pass false
+		MemoryImage.TargetLayoutParameters.InitializeForPlatform(CmdLineParameters.DataDrivenPlatformName, false);
+
+		FMemoryImageWriter Writer(MemoryImage);
+		Writer.WriteObject(Data);
+
+		FMemoryImageResult Result;
+		MemoryImage.Flatten(Result);
+
+		if (Info.bEncryptedIndex)
+		{
+			check(InKeyChain.MasterEncryptionKey);
+			Result.Bytes.AddZeroed(Align(Result.Bytes.Num(), FAES::AESBlockSize) - Result.Bytes.Num());
+			FSHA1::HashBuffer(Result.Bytes.GetData(), Result.Bytes.Num(), Info.IndexHash.Hash);
+			FAES::EncryptData(Result.Bytes.GetData(), Result.Bytes.Num(), InKeyChain.MasterEncryptionKey->Key);
+			TotalEncryptedDataSize += Result.Bytes.Num();
+		}
+
+		int32 Size = Result.Bytes.Num();
+		PakFileHandle->Serialize(Result.Bytes.GetData(), Size);
+		Result.SaveToArchive(*PakFileHandle);
+
+		// use this to store the frozen data size
+		Info.IndexSize = Size;
+		Info.bIndexIsFrozen = true;
 	}
+	else
+	{
+		UE_LOG(LogPakFile, Display, TEXT("DISABLING pak file index freezing (all must be true: HasPlatformInfo? %s - TargetIs64Bit? %s - NoRuntimeUnloading? %s - Unencrypted? %s)"),
+			PlatformInfo == nullptr ? TEXT("false") : TEXT("true"),
+			PlatformInfo == nullptr ? TEXT("unknown") : PlatformInfo->Freezing_b32Bit ? TEXT("false") : TEXT("true"),
+			CmdLineParameters.bAllowForIndexUnload ? TEXT("false") : TEXT("true"),
+			Info.bEncryptedIndex ? TEXT("false") : TEXT("true"));
+
+		// Serialize Pak Index at the end of Pak File
+		TArray<uint8> IndexData;
+		FMemoryWriter IndexWriter(IndexData);
+		IndexWriter.SetByteSwapping(PakFileHandle->ForceByteSwapping());
+		int32 NumEntries = Index.Num();
+		IndexWriter << MountPoint;
+		IndexWriter << NumEntries;
+		for (int32 EntryIndex = 0; EntryIndex < Index.Num(); EntryIndex++)
+		{
+			FPakEntryPair& Entry = Index[EntryIndex];
+			IndexWriter << Entry.Filename;
+			Entry.Info.Serialize(IndexWriter, Info.Version);
+
+			// @todo loadtime: can't this be done when writing the fileentry itself?
+			if (RequiredPatchPadding > 0)
+			{
+				int64 EntrySize = Entry.Info.GetSerializedSize(FPakInfo::PakFile_Version_Latest);
+				int64 TotalSizeToWrite = Entry.Info.Size + EntrySize;
+				if (TotalSizeToWrite >= RequiredPatchPadding)
+				{
+					int64 RealStart = Entry.Info.Offset;
+					if ((RealStart % RequiredPatchPadding) != 0 &&
+						!Entry.Filename.EndsWith(TEXT("uexp")) && // these are export sections of larger files and may be packed with uasset/umap and so we don't need a warning here
+						!(Entry.Filename.EndsWith(TEXT(".m.ubulk")) && CmdLineParameters.AlignForMemoryMapping > 0)) // Bulk padding unaligns patch padding and so we don't need a warning here
+					{
+						UE_LOG(LogPakFile, Warning, TEXT("File at offset %lld of size %lld not aligned to patch size %i"), RealStart, Entry.Info.Size, RequiredPatchPadding);
+					}
+				}
+			}
+		}
+
+		if (Info.bEncryptedIndex)
+		{
+			int32 OriginalSize = IndexData.Num();
+			int32 AlignedSize = Align(OriginalSize, FAES::AESBlockSize);
+
+			for (int32 PaddingIndex = IndexData.Num(); PaddingIndex < AlignedSize; ++PaddingIndex)
+			{
+				uint8 Byte = IndexData[PaddingIndex % OriginalSize];
+				IndexData.Add(Byte);
+			}
+		}
 
-	PakFileHandle->Serialize(IndexData.GetData(), IndexData.Num());
+		FSHA1::HashBuffer(IndexData.GetData(), IndexData.Num(), Info.IndexHash.Hash);
 
-	Info.IndexSize = IndexData.Num();
+		if (Info.bEncryptedIndex)
+		{
+			check(InKeyChain.MasterEncryptionKey);
+			FAES::EncryptData(IndexData.GetData(), IndexData.Num(), InKeyChain.MasterEncryptionKey->Key);
+			TotalEncryptedDataSize += IndexData.Num();
+		}
+
+		PakFileHandle->Serialize(IndexData.GetData(), IndexData.Num());
+
+		Info.IndexSize = IndexData.Num();
+	}
 
 	// Save trailer (offset, size, hash value)
 	Info.Serialize(*PakFileHandle, FPakInfo::PakFile_Version_Latest);
diff --git a/Engine/Source/Developer/PakFileUtilities/Private/SignedArchiveWriter.cpp b/Engine/Source/Developer/PakFileUtilities/Private/SignedArchiveWriter.cpp
--- a/Engine/Source/Developer/PakFileUtilities/Private/SignedArchiveWriter.cpp
+++ b/Engine/Source/Developer/PakFileUtilities/Private/SignedArchiveWriter.cpp
@@ -6,13 +6,15 @@
 #include "HAL/FileManager.h"
 
 FSignedArchiveWriter::FSignedArchiveWriter(FArchive& InPak, const FString& InPakFilename, const FRSAKeyHandle InSigningKey)
-: BufferArchive(Buffer)
-	, PakWriter(InPak)
-	, PakSignaturesFilename(FPaths::ChangeExtension(InPakFilename, TEXT("sig")))
-	, SizeOnDisk(0)
-	, PakSize(0)
-	, SigningKey(InSigningKey)
+: FArchive(InPak)
+, BufferArchive(Buffer)
+, PakWriter(InPak)
+, PakSignaturesFilename(FPaths::ChangeExtension(InPakFilename, TEXT("sig")))
+, SizeOnDisk(0)
+, PakSize(0)
+, SigningKey(InSigningKey)
 {
+	check(IsSaving());
 	Buffer.Reserve(FPakInfo::MaxChunkDataSize);
 }
 
diff --git a/Engine/Source/Developer/TargetPlatform/Public/Interfaces/IShaderFormat.h b/Engine/Source/Developer/TargetPlatform/Public/Interfaces/IShaderFormat.h
--- a/Engine/Source/Developer/TargetPlatform/Public/Interfaces/IShaderFormat.h
+++ b/Engine/Source/Developer/TargetPlatform/Public/Interfaces/IShaderFormat.h
@@ -7,6 +7,8 @@
 /** Wildcard string used to search for shader format modules. */
 #define SHADERFORMAT_MODULE_WILDCARD TEXT("*ShaderFormat*")
 
+class FSerializedShaderArchive;
+
 /**
  * IShaderFormat, shader pre-compilation abstraction
  */
@@ -70,9 +72,20 @@ public:
      * @param LibraryName The name of this shader library.
      * @param Format The format of shaders to cache.
      * @param WorkingDirectory The working directory.
-     * @returns An archive object on success or nullptr on failure.
+	 * @param The output directory for the archive.
+	 * @param The directory for the debug data previously generated.
+	 * @param Optional pointer to a TArray that on success will be filled with a list of the written file paths.
+     * @returns true if the archive was created
      */
-    virtual class IShaderFormatArchive* CreateShaderArchive( FString const& LibraryName, FName Format, const FString& WorkingDirectory ) const { return nullptr; }
+    virtual bool CreateShaderArchive( FString const& LibraryName,
+		FName Format,
+		const FString& WorkingDirectory,
+		const FString& OutputDir,
+		const FString& DebugOutputDir,
+		const FSerializedShaderArchive& SerializedShaders,
+		const TArray<TArray<uint8>>& ShaderCode,
+		TArray<FString>* OutputFiles) const
+	{ return false; }
 	
 	/**
 	 * Can the shader format compile shaders to the native binary format for the platform.
@@ -98,7 +111,7 @@ public:
 	 * Called when a shader resource is cooked, so the shader format can perform platform-specific operations on the debug data.
 	 * Does nothing on platforms that make no use of the platform debug data.
 	 */
-	virtual void NotifyShaderCooked(const TArray<uint8>& PlatformDebugData, FName Format) const { }
+	virtual void NotifyShaderCooked(const TConstArrayView<uint8>& PlatformDebugData, FName Format) const { }
 
 	/**
 	 * Appends shader key text to the provided key string for use in DDC unique key construction.
diff --git a/Engine/Source/Developer/TargetPlatform/Public/Interfaces/ITargetPlatform.h b/Engine/Source/Developer/TargetPlatform/Public/Interfaces/ITargetPlatform.h
--- a/Engine/Source/Developer/TargetPlatform/Public/Interfaces/ITargetPlatform.h
+++ b/Engine/Source/Developer/TargetPlatform/Public/Interfaces/ITargetPlatform.h
@@ -5,7 +5,6 @@
 #include "CoreMinimal.h"
 #include "Interfaces/ITargetDevice.h"
 
-
 namespace PlatformInfo
 {
 	// Forward declare type from DesktopPlatform rather than add an include dependency to everything using ITargetPlatform
@@ -461,13 +460,6 @@ public:
 	 */
 	virtual FName GetWaveFormat( const class USoundWave* Wave ) const = 0;
 
-	/**
-	* Get the audio compression settings for this platform.
-	* 
-	* @return the compression overrides for this platform, or the default platform overrides.
-	*/
-	virtual const struct FPlatformAudioCookOverrides* GetAudioCompressionSettings() const = 0;
-
 	/**
 	* Gets all the formats which can be returned from GetWaveFormat
 	*
diff --git a/Engine/Source/Developer/Windows/WindowsTargetPlatform/Classes/WindowsTargetSettings.h b/Engine/Source/Developer/Windows/WindowsTargetPlatform/Classes/WindowsTargetSettings.h
--- a/Engine/Source/Developer/Windows/WindowsTargetPlatform/Classes/WindowsTargetSettings.h
+++ b/Engine/Source/Developer/Windows/WindowsTargetPlatform/Classes/WindowsTargetSettings.h
@@ -73,6 +73,12 @@ public:
 	UPROPERTY(EditAnywhere, config, Category="OS Info", Meta=(DisplayName = "Minimum OS Version"))
 	EMinimumSupportedOS MinimumOSVersion;
 
+	/**
+	 * Determines if data is cooked for 32bit (or 64bit) exe
+	 */
+	UPROPERTY(EditAnywhere, config, Category = "OS Info", Meta = (DisplayName = "Target 32bit"))
+	bool bTarget32Bit;
+
 	/** Sample rate to run the audio mixer with. */
 	UPROPERTY(config, EditAnywhere, Category = "Audio", Meta = (DisplayName = "Audio Mixer Sample Rate"))
 	int32 AudioSampleRate;
diff --git a/Engine/Source/Developer/Windows/WindowsTargetPlatform/Private/GenericWindowsTargetPlatform.h b/Engine/Source/Developer/Windows/WindowsTargetPlatform/Private/GenericWindowsTargetPlatform.h
--- a/Engine/Source/Developer/Windows/WindowsTargetPlatform/Private/GenericWindowsTargetPlatform.h
+++ b/Engine/Source/Developer/Windows/WindowsTargetPlatform/Private/GenericWindowsTargetPlatform.h
@@ -6,6 +6,7 @@
 #include "Common/TargetPlatformBase.h"
 #include "Misc/ConfigCacheIni.h"
 #include "LocalPcTargetDevice.h"
+#include "Serialization/MemoryLayout.h"
 
 #if WITH_ENGINE
 	#include "Sound/SoundWave.h"
@@ -19,127 +20,6 @@
 
 #define LOCTEXT_NAMESPACE "TGenericWindowsTargetPlatform"
 
-namespace Windows
-{
-#if WITH_ENGINE
-	FORCEINLINE void CachePlatformAudioCookOverrides(FPlatformAudioCookOverrides& OutOverrides)
-	{
-		const TCHAR* CategoryName = TEXT("/Script/WindowsTargetPlatform.WindowsTargetSettings");
-
-		int32 SoundCueQualityIndex = INDEX_NONE;
-		if (GConfig->GetInt(CategoryName, TEXT("SoundCueCookQualityIndex"), SoundCueQualityIndex, GEngineIni))
-		{
-			OutOverrides.SoundCueCookQualityIndex = SoundCueQualityIndex;
-		}
-
-		GConfig->GetBool(CategoryName, TEXT("bUseAudioStreamCaching"), OutOverrides.bUseStreamCaching, GEngineIni);
-
-		/** Memory Load On Demand Settings */
-		if (OutOverrides.bUseStreamCaching)
-		{
-			// Cache size:
-			int32 RetrievedCacheSize = 32 * 1024;
-			GConfig->GetInt(CategoryName, TEXT("CacheSizeKB"), RetrievedCacheSize, GEngineIni);
-			OutOverrides.StreamCachingSettings.CacheSizeKB = RetrievedCacheSize;
-		}
-
-		GConfig->GetBool(CategoryName, TEXT("bResampleForDevice"), OutOverrides.bResampleForDevice, GEngineIni);
-
-		GConfig->GetFloat(CategoryName, TEXT("CompressionQualityModifier"), OutOverrides.CompressionQualityModifier, GEngineIni);
-
-		GConfig->GetFloat(CategoryName, TEXT("AutoStreamingThreshold"), OutOverrides.AutoStreamingThreshold, GEngineIni);
-
-		//Cache sample rate map:
-		float RetrievedSampleRate = -1.0f;
-
-		GConfig->GetFloat(CategoryName, TEXT("MaxSampleRate"), RetrievedSampleRate, GEngineIni);
-		float* FoundSampleRate = OutOverrides.PlatformSampleRates.Find(ESoundwaveSampleRateSettings::Max);
-
-		if (FoundSampleRate)
-		{
-			if (!FMath::IsNearlyEqual(*FoundSampleRate, RetrievedSampleRate))
-			{
-				*FoundSampleRate = RetrievedSampleRate;
-			}
-
-		}
-		else
-		{
-			OutOverrides.PlatformSampleRates.Add(ESoundwaveSampleRateSettings::Max, RetrievedSampleRate);
-		}
-
-		RetrievedSampleRate = -1.0f;
-
-		GConfig->GetFloat(CategoryName, TEXT("HighSampleRate"), RetrievedSampleRate, GEngineIni);
-		FoundSampleRate = OutOverrides.PlatformSampleRates.Find(ESoundwaveSampleRateSettings::High);
-
-		if (FoundSampleRate)
-		{
-			if (!FMath::IsNearlyEqual(*FoundSampleRate, RetrievedSampleRate))
-			{
-				*FoundSampleRate = RetrievedSampleRate;
-			}
-
-		}
-		else
-		{
-			OutOverrides.PlatformSampleRates.Add(ESoundwaveSampleRateSettings::High, RetrievedSampleRate);
-		}
-
-
-		RetrievedSampleRate = -1.0f;
-
-		GConfig->GetFloat(CategoryName, TEXT("MedSampleRate"), RetrievedSampleRate, GEngineIni);
-		FoundSampleRate = OutOverrides.PlatformSampleRates.Find(ESoundwaveSampleRateSettings::Medium);
-
-		if (FoundSampleRate)
-		{
-			if (!FMath::IsNearlyEqual(*FoundSampleRate, RetrievedSampleRate))
-			{
-				*FoundSampleRate = RetrievedSampleRate;
-			}
-		}
-		else
-		{
-			OutOverrides.PlatformSampleRates.Add(ESoundwaveSampleRateSettings::Medium, RetrievedSampleRate);
-		}
-
-		RetrievedSampleRate = -1.0f;
-
-		GConfig->GetFloat(CategoryName, TEXT("LowSampleRate"), RetrievedSampleRate, GEngineIni);
-		FoundSampleRate = OutOverrides.PlatformSampleRates.Find(ESoundwaveSampleRateSettings::Low);
-
-		if (FoundSampleRate)
-		{
-			if (!FMath::IsNearlyEqual(*FoundSampleRate, RetrievedSampleRate))
-			{
-				*FoundSampleRate = RetrievedSampleRate;
-			}
-		}
-		else
-		{
-			OutOverrides.PlatformSampleRates.Add(ESoundwaveSampleRateSettings::Low, RetrievedSampleRate);
-		}
-
-		RetrievedSampleRate = -1.0f;
-
-		GConfig->GetFloat(CategoryName, TEXT("MinSampleRate"), RetrievedSampleRate, GEngineIni);
-		FoundSampleRate = OutOverrides.PlatformSampleRates.Find(ESoundwaveSampleRateSettings::Min);
-
-		if (FoundSampleRate)
-		{
-			if (!FMath::IsNearlyEqual(*FoundSampleRate, RetrievedSampleRate))
-			{
-				*FoundSampleRate = RetrievedSampleRate;
-			}
-		}
-		else
-		{
-			OutOverrides.PlatformSampleRates.Add(ESoundwaveSampleRateSettings::Min, RetrievedSampleRate);
-		}
-	}
-#endif
-}
 
 /**
  * Template for Windows target platforms
@@ -550,21 +430,6 @@ public:
 		OutFormats.Add(NAME_OPUS);
 	}
 
-	virtual FPlatformAudioCookOverrides* GetAudioCompressionSettings() const override
-	{
-		static FPlatformAudioCookOverrides Settings;
-
-		static bool bCachedPlatformSettings = false;
-
-		if (!bCachedPlatformSettings)
-		{
-			Windows::CachePlatformAudioCookOverrides(Settings);
-			bCachedPlatformSettings = true;
-		}
-
-		return &Settings;
-	}
-
 #endif //WITH_ENGINE
 
 	virtual bool SupportsVariants() const override
diff --git a/Engine/Source/Developer/Windows/WindowsTargetPlatform/Private/WindowsTargetPlatformClasses.cpp b/Engine/Source/Developer/Windows/WindowsTargetPlatform/Private/WindowsTargetPlatformClasses.cpp
--- a/Engine/Source/Developer/Windows/WindowsTargetPlatform/Private/WindowsTargetPlatformClasses.cpp
+++ b/Engine/Source/Developer/Windows/WindowsTargetPlatform/Private/WindowsTargetPlatformClasses.cpp
@@ -9,6 +9,7 @@ UWindowsTargetSettings::UWindowsTargetSettings( const FObjectInitializer& Object
 	: Super(ObjectInitializer)
 {
 	MinimumOSVersion = EMinimumSupportedOS::MSOS_Vista;
+	bTarget32Bit = false;
 
 	// Default windows settings
 	AudioSampleRate = 48000;
diff --git a/Engine/Source/Editor/AudioEditor/Private/Factories/SoundFactory.cpp b/Engine/Source/Editor/AudioEditor/Private/Factories/SoundFactory.cpp
--- a/Engine/Source/Editor/AudioEditor/Private/Factories/SoundFactory.cpp
+++ b/Engine/Source/Editor/AudioEditor/Private/Factories/SoundFactory.cpp
@@ -496,7 +496,7 @@ UObject* USoundFactory::CreateObject
 		Sound->InvalidateCompressedData(true /* bFreeResources */);
 
 		// If stream caching is enabled, we need to make sure this asset is ready for playback.
-		if (FPlatformCompressionUtilities::IsCurrentPlatformUsingStreamCaching() && Sound->IsStreaming())
+		if (FPlatformCompressionUtilities::IsCurrentPlatformUsingStreamCaching() && Sound->IsStreaming(nullptr))
 		{
 			Sound->EnsureZerothChunkIsLoaded();
 		}
diff --git a/Engine/Source/Editor/FoliageEdit/Private/FoliageEdMode.cpp b/Engine/Source/Editor/FoliageEdit/Private/FoliageEdMode.cpp
--- a/Engine/Source/Editor/FoliageEdit/Private/FoliageEdMode.cpp
+++ b/Engine/Source/Editor/FoliageEdit/Private/FoliageEdMode.cpp
@@ -3166,10 +3166,8 @@ bool FEdModeFoliage::GetStaticMeshVertexColorForHit(const UStaticMeshComponent*
 	const FPositionVertexBuffer& PositionVertexBuffer = LODModel.VertexBuffers.PositionVertexBuffer;
 
 	int32 SectionFirstTriIndex = 0;
-	for (TArray<FStaticMeshSection>::TConstIterator SectionIt(LODModel.Sections); SectionIt; ++SectionIt)
+	for (const FStaticMeshSection& Section : LODModel.Sections)
 	{
-		const FStaticMeshSection& Section = *SectionIt;
-
 		if (Section.bEnableCollision)
 		{
 			int32 NextSectionTriIndex = SectionFirstTriIndex + Section.NumTriangles;
diff --git a/Engine/Source/Editor/KismetCompiler/Private/KismetCompilerVMBackend.cpp b/Engine/Source/Editor/KismetCompiler/Private/KismetCompilerVMBackend.cpp
--- a/Engine/Source/Editor/KismetCompiler/Private/KismetCompilerVMBackend.cpp
+++ b/Engine/Source/Editor/KismetCompiler/Private/KismetCompilerVMBackend.cpp
@@ -905,9 +905,9 @@ public:
 				ensure(InnerProp);
 				FScriptArray ScriptArray;
 				ArrayPropr->ImportText(*Term->Name, &ScriptArray, 0, NULL, GLog);
-				int32 ElementNum = ScriptArray.Num();
 
 				FScriptArrayHelper ScriptArrayHelper(ArrayPropr, &ScriptArray);
+				int32 ElementNum = ScriptArrayHelper.Num();
 
 				Writer << EX_ArrayConst;
 				Writer << InnerProp;
diff --git a/Engine/Source/Editor/MaterialEditor/Private/MaterialEditor.cpp b/Engine/Source/Editor/MaterialEditor/Private/MaterialEditor.cpp
--- a/Engine/Source/Editor/MaterialEditor/Private/MaterialEditor.cpp
+++ b/Engine/Source/Editor/MaterialEditor/Private/MaterialEditor.cpp
@@ -3499,9 +3499,9 @@ void FMaterialEditor::SetVectorParameterDefaultOnDependentMaterials(FName Parame
 			{
 				bool bUpdate = false;
 
-				for (int32 FunctionIndex = 0; FunctionIndex < CurrentMaterial->MaterialFunctionInfos.Num(); FunctionIndex++)
+				for (int32 FunctionIndex = 0; FunctionIndex < CurrentMaterial->GetCachedExpressionData().FunctionInfos.Num(); FunctionIndex++)
 				{
-					if (CurrentMaterial->MaterialFunctionInfos[FunctionIndex].Function == MaterialFunction->ParentFunction)
+					if (CurrentMaterial->GetCachedExpressionData().FunctionInfos[FunctionIndex].Function == MaterialFunction->ParentFunction)
 					{
 						bUpdate = true;
 						break;
@@ -3575,9 +3575,9 @@ void FMaterialEditor::SetScalarParameterDefaultOnDependentMaterials(FName Parame
 			{
 				bool bUpdate = false;
 
-				for (int32 FunctionIndex = 0; FunctionIndex < CurrentMaterial->MaterialFunctionInfos.Num(); FunctionIndex++)
+				for (int32 FunctionIndex = 0; FunctionIndex < CurrentMaterial->GetCachedExpressionData().FunctionInfos.Num(); FunctionIndex++)
 				{
-					if (CurrentMaterial->MaterialFunctionInfos[FunctionIndex].Function == MaterialFunction->ParentFunction)
+					if (CurrentMaterial->GetCachedExpressionData().FunctionInfos[FunctionIndex].Function == MaterialFunction->ParentFunction)
 					{
 						bUpdate = true;
 						break;
diff --git a/Engine/Source/Editor/MaterialEditor/Private/MaterialEditor.h b/Engine/Source/Editor/MaterialEditor/Private/MaterialEditor.h
--- a/Engine/Source/Editor/MaterialEditor/Private/MaterialEditor.h
+++ b/Engine/Source/Editor/MaterialEditor/Private/MaterialEditor.h
@@ -63,7 +63,7 @@ public:
 		FPlatformMisc::CreateGuid(Id);
 
 		check(InExpression->Material && InExpression->Material->Expressions.Contains(InExpression));
-		InExpression->Material->AppendReferencedTextures(ReferencedTextures);
+		ReferencedTextures = InExpression->Material->GetReferencedTextures();
 		SetQualityLevelProperties(EMaterialQualityLevel::High, false, GMaxRHIFeatureLevel);
 	}
 
@@ -93,9 +93,9 @@ public:
 	 */
 	virtual bool ShouldCache(EShaderPlatform Platform, const FShaderType* ShaderType, const FVertexFactoryType* VertexFactoryType) const override;
 
-	virtual const TArray<UObject*>& GetReferencedTextures() const override
+	virtual TArrayView<UObject* const> GetReferencedTextures() const override
 	{
-		return ReferencedTextures;
+		return MakeArrayView(ReferencedTextures);
 	}
 
 	////////////////
@@ -114,7 +114,7 @@ public:
 		}
 	}
 
-	virtual bool GetVectorValue(const FMaterialParameterInfo& ParameterInfo, FLinearColor* OutValue, const FMaterialRenderContext& Context) const override
+	virtual bool GetVectorValue(const FHashedMaterialParameterInfo& ParameterInfo, FLinearColor* OutValue, const FMaterialRenderContext& Context) const override
 	{
 		if (Expression.IsValid() && Expression->Material)
 		{
@@ -123,7 +123,7 @@ public:
 		return false;
 	}
 
-	virtual bool GetScalarValue(const FMaterialParameterInfo& ParameterInfo, float* OutValue, const FMaterialRenderContext& Context) const override
+	virtual bool GetScalarValue(const FHashedMaterialParameterInfo& ParameterInfo, float* OutValue, const FMaterialRenderContext& Context) const override
 	{
 		if (Expression.IsValid() && Expression->Material)
 		{
@@ -132,7 +132,7 @@ public:
 		return false;
 	}
 
-	virtual bool GetTextureValue(const FMaterialParameterInfo& ParameterInfo,const UTexture** OutValue, const FMaterialRenderContext& Context) const override
+	virtual bool GetTextureValue(const FHashedMaterialParameterInfo& ParameterInfo,const UTexture** OutValue, const FMaterialRenderContext& Context) const override
 	{
 		if (Expression.IsValid() && Expression->Material)
 		{
@@ -141,7 +141,7 @@ public:
 		return false;
 	}
 
-	virtual bool GetTextureValue(const FMaterialParameterInfo& ParameterInfo, const URuntimeVirtualTexture** OutValue, const FMaterialRenderContext& Context) const override
+	virtual bool GetTextureValue(const FHashedMaterialParameterInfo& ParameterInfo, const URuntimeVirtualTexture** OutValue, const FMaterialRenderContext& Context) const override
 	{
 		if (Expression.IsValid() && Expression->Material)
 		{
diff --git a/Engine/Source/Editor/MaterialEditor/Private/MaterialStats.cpp b/Engine/Source/Editor/MaterialEditor/Private/MaterialStats.cpp
--- a/Engine/Source/Editor/MaterialEditor/Private/MaterialStats.cpp
+++ b/Engine/Source/Editor/MaterialEditor/Private/MaterialStats.cpp
@@ -101,16 +101,16 @@ FText FShaderPlatformSettings::GetShaderCode(const EMaterialQualityLevel::Type Q
 	// check if shader compilation is done and extract shader code
 	if (bCompilationFinished)
 	{
-		TMap<FName, FShader*> ShaderMap;
+		TMap<FHashedName, TShaderRef<FShader>> ShaderMap;
 		MaterialShaderMap->GetShaderList(ShaderMap);
 
 		const auto Entry = ShaderMap.Find(PlatformData[QualityType].ComboBoxSelectedName);
 		if (Entry != nullptr)
 		{
-			const FShader* Shader = *Entry;
+			const TShaderRef<FShader>& Shader = *Entry;
 
-			const FName ShaderFName = Shader->GetType()->GetFName();
-			const FString* ShaderSource = MaterialShaderMap->GetShaderSource(ShaderFName);
+			const FName ShaderFName = Shader.GetType()->GetFName();
+			const FMemoryImageString* ShaderSource = MaterialShaderMap->GetShaderSource(ShaderFName);
 			if (ShaderSource != nullptr)
 			{
 				PlatformData[QualityType].bUpdateShaderCode = false;
@@ -173,12 +173,11 @@ bool FShaderPlatformSettings::CheckShaders()
 			{
 				Data.MaterialResourcesStats->CancelCompilation();
 
-				Material->RebuildExpressionTextureReferences();
+				Material->UpdateCachedExpressionData();
 
 				if (MaterialInstance != nullptr)
 				{
-					MaterialInstance->PermutationTextureReferences.Empty();
-					MaterialInstance->AppendReferencedTextures(MaterialInstance->PermutationTextureReferences);
+					MaterialInstance->UpdateCachedLayerParameters();
 				}
 
 				Data.MaterialResourcesStats->CacheShaders(PlatformShaderID);
@@ -219,13 +218,13 @@ bool FShaderPlatformSettings::Update()
 				const FMaterialShaderMap* MaterialShaderMap = Resource->GetGameThreadShaderMap();
 				if (MaterialShaderMap != nullptr)
 				{
-					TMap<FShaderId, FShader*> ShaderMap;
+					TMap<FShaderId, TShaderRef<FShader>> ShaderMap;
 					MaterialShaderMap->GetShaderList(ShaderMap);
 
 					QualityItem.ArrShaderNames.Empty();
 					for (const auto Entry : ShaderMap)
 					{
-						QualityItem.ArrShaderNames.Add(MakeShareable(new FName(Entry.Key.ShaderType->GetFName())));
+						QualityItem.ArrShaderNames.Add(MakeShareable(new FName(Entry.Value.GetType()->GetFName())));
 					}
 
 					if (QualityItem.ArrShaderNames.Num() > 0)
diff --git a/Engine/Source/Editor/MaterialEditor/Private/MaterialStatsCommon.cpp b/Engine/Source/Editor/MaterialEditor/Private/MaterialStatsCommon.cpp
--- a/Engine/Source/Editor/MaterialEditor/Private/MaterialStatsCommon.cpp
+++ b/Engine/Source/Editor/MaterialEditor/Private/MaterialStatsCommon.cpp
@@ -237,7 +237,7 @@ bool FMaterialStatsUtils::PlatformNeedsOfflineCompiler(const EShaderPlatform Sha
 			return false;
 
 		default:
-			return FDataDrivenShaderPlatformInfo::GetInfo(ShaderPlatform).bNeedsOfflineCompiler;
+			return FDataDrivenShaderPlatformInfo::GetNeedsOfflineCompiler(ShaderPlatform);
 	}
 
 	return false;
@@ -626,8 +626,8 @@ void FMaterialStatsUtils::GetRepresentativeInstructionCounts(TArray<FShaderInstr
 				const FMeshMaterialShaderMap* MeshShaderMap = MaterialShaderMap->GetMeshShaderMap(FactoryType);
 				if (MeshShaderMap)
 				{
-					TMap<FName, FShader*> ShaderMap;
-					MeshShaderMap->GetShaderList(ShaderMap);
+					TMap<FHashedName, TShaderRef<FShader>> ShaderMap;
+					MeshShaderMap->GetShaderList(*MaterialShaderMap, ShaderMap);
 
 					auto& DescriptionArray = DescriptionPair.Value;
 
@@ -635,12 +635,12 @@ void FMaterialStatsUtils::GetRepresentativeInstructionCounts(TArray<FShaderInstr
 					{
 						const FRepresentativeShaderInfo& ShaderInfo = DescriptionArray[i];
 
-						FShader** ShaderEntry = ShaderMap.Find(ShaderInfo.ShaderName);
+						TShaderRef<FShader>* ShaderEntry = ShaderMap.Find(ShaderInfo.ShaderName);
 						if (ShaderEntry != nullptr)
 						{
-							FShaderType* ShaderType = (*ShaderEntry)->GetType();
+							FShaderType* ShaderType = (*ShaderEntry).GetType();
 							{
-								const int32 NumInstructions = MeshShaderMap->GetMaxNumInstructionsForShader(ShaderType);
+								const int32 NumInstructions = MeshShaderMap->GetMaxNumInstructionsForShader(*MaterialShaderMap, ShaderType);
 
 								FShaderInstructionsInfo Info;
 								Info.ShaderType = ShaderInfo.ShaderType;
diff --git a/Engine/Source/Editor/MovieSceneTools/Private/TrackEditors/MaterialTrackEditor.cpp b/Engine/Source/Editor/MovieSceneTools/Private/TrackEditors/MaterialTrackEditor.cpp
--- a/Engine/Source/Editor/MovieSceneTools/Private/TrackEditors/MaterialTrackEditor.cpp
+++ b/Engine/Source/Editor/MovieSceneTools/Private/TrackEditors/MaterialTrackEditor.cpp
@@ -88,7 +88,7 @@ TSharedRef<SWidget> FMaterialTrackEditor::OnGetAddParameterMenuContent( FGuid Ob
 		Material->GetAllScalarParameterInfo(ScalarParameterInfo, ScalarParameterGuids );
 		for (int32 ScalarParameterIndex = 0; ScalarParameterIndex < ScalarParameterInfo.Num(); ++ScalarParameterIndex)
 		{
-			if (!bCollectedVisibleParameters || VisibleExpressions.Contains(ScalarParameterInfo[ScalarParameterIndex].Name))
+			if (!bCollectedVisibleParameters || VisibleExpressions.Contains(FMaterialParameterInfo(ScalarParameterInfo[ScalarParameterIndex].Name)))
 			{
 				FName ScalarParameterName = ScalarParameterInfo[ScalarParameterIndex].Name;
 				FUIAction AddParameterMenuAction( FExecuteAction::CreateSP( this, &FMaterialTrackEditor::AddScalarParameter, ObjectBinding, MaterialTrack, ScalarParameterName ) );
@@ -103,7 +103,7 @@ TSharedRef<SWidget> FMaterialTrackEditor::OnGetAddParameterMenuContent( FGuid Ob
 		Material->GetAllVectorParameterInfo( ColorParameterInfo, ColorParameterGuids );
 		for (int32 ColorParameterIndex = 0; ColorParameterIndex < ColorParameterInfo.Num(); ++ColorParameterIndex)
 		{
-			if (!bCollectedVisibleParameters || VisibleExpressions.Contains(ColorParameterInfo[ColorParameterIndex].Name))
+			if (!bCollectedVisibleParameters || VisibleExpressions.Contains(FMaterialParameterInfo(ColorParameterInfo[ColorParameterIndex].Name)))
 			{
 				FName ColorParameterName = ColorParameterInfo[ColorParameterIndex].Name;
 				FUIAction AddParameterMenuAction( FExecuteAction::CreateSP( this, &FMaterialTrackEditor::AddColorParameter, ObjectBinding, MaterialTrack, ColorParameterName ) );
diff --git a/Engine/Source/Editor/PropertyEditor/Private/ItemPropertyNode.cpp b/Engine/Source/Editor/PropertyEditor/Private/ItemPropertyNode.cpp
--- a/Engine/Source/Editor/PropertyEditor/Private/ItemPropertyNode.cpp
+++ b/Engine/Source/Editor/PropertyEditor/Private/ItemPropertyNode.cpp
@@ -199,7 +199,8 @@ void FItemPropertyNode::InitChildNodes()
 
 		if( Array )
 		{
-			for( int32 Index = 0 ; Index < FScriptArrayHelper::Num(Array) ; Index++ )
+			FScriptArrayHelper ArrayHelper(ArrayProperty, Array);
+			for( int32 Index = 0 ; Index < ArrayHelper.Num(); Index++ )
 			{
 				TSharedPtr<FItemPropertyNode> NewItemNode( new FItemPropertyNode );
 
diff --git a/Engine/Source/Editor/PropertyEditor/Private/ObjectPropertyNode.cpp b/Engine/Source/Editor/PropertyEditor/Private/ObjectPropertyNode.cpp
--- a/Engine/Source/Editor/PropertyEditor/Private/ObjectPropertyNode.cpp
+++ b/Engine/Source/Editor/PropertyEditor/Private/ObjectPropertyNode.cpp
@@ -213,13 +213,15 @@ bool FObjectPropertyNode::GetReadAddressUncached(FPropertyNode& InNode,
 			{
 				if ( ArrayOuter )
 				{
-					const int32 Num = FScriptArrayHelper::Num(BaseAddr);
+					FScriptArrayHelper ArrayHelper(ArrayOuter, BaseAddr);
+					const int32 Num = ArrayHelper.Num();
 					for (int32 ObjIndex = 1; ObjIndex < GetNumObjects(); ++ObjIndex)
 					{
 						TempObject = GetUObject(ObjIndex);
 						BaseAddr = ParentPropertyNode->GetValueBaseAddressFromObject(TempObject);
 
-						if (BaseAddr && Num != FScriptArrayHelper::Num(BaseAddr))
+						ArrayHelper = FScriptArrayHelper(ArrayOuter, BaseAddr);
+						if (BaseAddr && Num != ArrayHelper.Num())
 						{
 							bAllTheSame = false;
 						}
@@ -241,15 +243,20 @@ bool FObjectPropertyNode::GetReadAddressUncached(FPropertyNode& InNode,
 				}
 				else if ( MapOuter )
 				{
-					const int32 Num = FScriptMapHelper::Num(BaseAddr);
+					FScriptMapHelper MapHelper(MapOuter, BaseAddr);
+					const int32 Num = MapHelper.Num();
 					for (int32 ObjIndex = 1; ObjIndex < GetNumObjects(); ++ObjIndex)
 					{
 						TempObject = GetUObject(ObjIndex);
 						BaseAddr = ParentPropertyNode->GetValueBaseAddressFromObject(TempObject);
 
-						if (BaseAddr && Num != FScriptMapHelper::Num(BaseAddr))
+						if (BaseAddr)
 						{
-							bAllTheSame = false;
+							MapHelper = FScriptMapHelper(MapOuter, BaseAddr);
+							if (Num != MapHelper.Num())
+							{
+								bAllTheSame = false;
+							}
 						}
 					}
 				}
@@ -277,13 +284,18 @@ bool FObjectPropertyNode::GetReadAddressUncached(FPropertyNode& InNode,
 
 				if (ArrayProp)
 				{
-					int32 const Num = FScriptArrayHelper::Num(BaseAddr);
+					FScriptArrayHelper ArrayHelper(ArrayProp, BaseAddr);
+					int32 const Num = ArrayHelper.Num();
 					for (int32 ObjIndex = 1; ObjIndex < GetNumObjects(); ObjIndex++)
 					{
 						TempObject = GetUObject(ObjIndex);
-						if (TempObject && Num != FScriptArrayHelper::Num(InNode.GetValueBaseAddressFromObject(TempObject)))
+						if (TempObject)
 						{
-							bAllTheSame = false;
+							ArrayHelper = FScriptArrayHelper(ArrayProp, InNode.GetValueBaseAddressFromObject(TempObject));
+							if (Num != ArrayHelper.Num())
+							{
+								bAllTheSame = false;
+							}
 						}
 					}
 				}
@@ -301,13 +313,18 @@ bool FObjectPropertyNode::GetReadAddressUncached(FPropertyNode& InNode,
 				}
 				else if (MapProp)
 				{
-					int32 const Num = FScriptMapHelper::Num(BaseAddr);
+					FScriptMapHelper MapHelper(MapProp, BaseAddr);
+					int32 const Num = MapHelper.Num();
 					for (int32 ObjIndex = 1; ObjIndex < GetNumObjects(); ++ObjIndex)
 					{
 						TempObject = GetUObject(ObjIndex);
-						if (TempObject && Num != FScriptMapHelper::Num(InNode.GetValueBaseAddressFromObject(TempObject)))
+						if (TempObject)
 						{
-							bAllTheSame = false;
+							MapHelper = FScriptMapHelper(MapProp, InNode.GetValueBaseAddressFromObject(TempObject));
+							if (Num != MapHelper.Num())
+							{
+								bAllTheSame = false;
+							}
 						}
 					}
 				}
diff --git a/Engine/Source/Editor/PropertyEditor/Private/PropertyChangeListener.cpp b/Engine/Source/Editor/PropertyEditor/Private/PropertyChangeListener.cpp
--- a/Engine/Source/Editor/PropertyEditor/Private/PropertyChangeListener.cpp
+++ b/Engine/Source/Editor/PropertyEditor/Private/PropertyChangeListener.cpp
@@ -113,7 +113,8 @@ public:
 			if ( OuterArrayProperty != NULL )
 			{
 				// make sure we're not trying to compare against an element that doesn't exist
-				if ( PropertyNodeRef.GetArrayIndex() >= FScriptArrayHelper::Num( PropertyValueAddresses.BaseAddress ) )
+				FScriptArrayHelper ArrayHelper(OuterArrayProperty, PropertyValueAddresses.BaseAddress);
+				if ( PropertyNodeRef.GetArrayIndex() >= ArrayHelper.Num() )
 				{
 					bPropertyValid = false;
 				}
diff --git a/Engine/Source/Editor/PropertyEditor/Private/PropertyHandleImpl.cpp b/Engine/Source/Editor/PropertyEditor/Private/PropertyHandleImpl.cpp
--- a/Engine/Source/Editor/PropertyEditor/Private/PropertyHandleImpl.cpp
+++ b/Engine/Source/Editor/PropertyEditor/Private/PropertyHandleImpl.cpp
@@ -337,17 +337,19 @@ FString FPropertyValueImpl::GetPropertyValueArray() const
 				uint8* Addr = ReadAddresses.GetAddress(0);
 				if( Addr )
 				{
-					if ( NodeProperty->IsA<FArrayProperty>() )
+					if ( FArrayProperty* ArrayProperty = CastField<FArrayProperty>(NodeProperty) )
 					{
-						String = FString::Printf( TEXT("%(%d)"), FScriptArrayHelper::Num(Addr) );
+						FScriptArrayHelper ArrayHelper(ArrayProperty, Addr);
+						String = FString::Printf( TEXT("%(%d)"), ArrayHelper.Num() );
 					}
 					else if ( CastField<FSetProperty>(NodeProperty) != nullptr )	
 					{
 						String = FString::Printf( TEXT("%(%d)"), FScriptSetHelper::Num(Addr) );
 					}
-					else if (CastField<FMapProperty>(NodeProperty) != nullptr)
+					else if (FMapProperty* MapProperty = CastField<FMapProperty>(NodeProperty))
 					{
-						String = FString::Printf(TEXT("%(%d)"), FScriptMapHelper::Num(Addr));
+						FScriptMapHelper MapHelper(MapProperty, Addr);
+						String = FString::Printf(TEXT("%(%d)"), MapHelper.Num());
 					}
 					else
 					{
@@ -772,7 +774,8 @@ static int32 GetArrayPropertyLastValidIndex( FObjectPropertyNode* InObjectNode,
 		if (ArrayProperty)
 		{
 			uint8* PropertyAddressBase = ArrayProperty->ContainerPtrToValuePtr<uint8>(ParentObject);
-			ClampMax = FScriptArrayHelper::Num(PropertyAddressBase) - 1;
+			FScriptArrayHelper ArrayHelper(ArrayProperty, PropertyAddressBase);
+			ClampMax = ArrayHelper.Num() - 1;
 		}
 		else
 		{
diff --git a/Engine/Source/Editor/PropertyEditor/Private/PropertyNode.cpp b/Engine/Source/Editor/PropertyEditor/Private/PropertyNode.cpp
--- a/Engine/Source/Editor/PropertyEditor/Private/PropertyNode.cpp
+++ b/Engine/Source/Editor/PropertyEditor/Private/PropertyNode.cpp
@@ -511,7 +511,8 @@ EPropertyDataValidationResult FPropertyNode::EnsureDataIsValid()
 				if( ArrayProperty && !bIgnoreAllMismatch)
 				{
 					//ensure that array structures have the proper number of children
-					int32 ArrayNum = FScriptArrayHelper::Num(Addr);
+					FScriptArrayHelper ArrayHelper(ArrayProperty, Addr);
+					int32 ArrayNum = ArrayHelper.Num();
 					//if first child
 					if (NumArrayChildren == -1)
 					{
@@ -541,7 +542,8 @@ EPropertyDataValidationResult FPropertyNode::EnsureDataIsValid()
 
 				if (MapProperty && !bIgnoreAllMismatch)
 				{
-					int32 MapNum = FScriptMapHelper::Num(Addr);
+					FScriptMapHelper MapHelper(MapProperty, Addr);
+					int32 MapNum = MapHelper.Num();
 
 					if (NumArrayChildren == -1)
 					{
@@ -1369,10 +1371,11 @@ private:
 
 		if ( ArrayProp != NULL )
 		{
-			FScriptArray* ArrayValuePtr = ArrayProp->GetPropertyValuePtr(PropertyValueAddress);
+			FScriptArrayHelper ArrayHelper(ArrayProp, PropertyValueAddress);
 
-			uint8* ArrayValue = (uint8*)ArrayValuePtr->GetData();
-			for ( int32 ArrayIndex = 0; ArrayIndex < ArrayValuePtr->Num(); ArrayIndex++ )
+			uint8* ArrayValue = (uint8*)ArrayHelper.GetRawPtr();
+			int32 ArraySize = ArrayHelper.Num();
+			for ( int32 ArrayIndex = 0; ArrayIndex < ArraySize; ArrayIndex++ )
 			{
 				ProcessProperty(ArrayProp->Inner, ArrayValue + ArrayIndex * ArrayProp->Inner->ElementSize);
 			}
@@ -1431,18 +1434,16 @@ private:
 
 		if (MapProp != NULL)
 		{
-			FScriptMap* MapValuePtr = MapProp->GetPropertyValuePtr(PropertyValueAddress);
-
-			FScriptMapLayout MapLayout = MapValuePtr->GetScriptLayout(MapProp->KeyProp->ElementSize, MapProp->KeyProp->GetMinAlignment(), MapProp->ValueProp->ElementSize, MapProp->ValueProp->GetMinAlignment());
-			int32 ItemsLeft = MapValuePtr->Num();
+			FScriptMapHelper MapHelper(MapProp, PropertyValueAddress);
 
+			int32 ItemsLeft = MapHelper.Num();
 			for (int32 Index = 0; ItemsLeft > 0; ++Index)
 			{
-				if (MapValuePtr->IsValidIndex(Index))
+				if (MapHelper.IsValidIndex(Index))
 				{
 					--ItemsLeft;
 
-					uint8* Data = (uint8*)MapValuePtr->GetData(Index, MapLayout);
+					uint8* Data = MapHelper.GetPairPtr(Index);
 
 					ProcessProperty(MapProp->KeyProp, MapProp->KeyProp->ContainerPtrToValuePtr<uint8>(Data));
 					ProcessProperty(MapProp->ValueProp, MapProp->ValueProp->ContainerPtrToValuePtr<uint8>(Data));
@@ -1617,7 +1618,8 @@ bool FPropertyNode::GetDiffersFromDefaultForObject( FPropertyItemValueDataTracke
 			if (ValueTracker.GetPropertyDefaultBaseAddress() != nullptr)
 			{
 				// make sure we're not trying to compare against an element that doesn't exist
-				if (GetArrayIndex() >= FScriptArrayHelper::Num(ValueTracker.GetPropertyDefaultBaseAddress()))
+				FScriptArrayHelper ArrayHelper(OuterArrayProperty, ValueTracker.GetPropertyDefaultBaseAddress());
+				if (GetArrayIndex() >= ArrayHelper.Num())
 				{
 					bDiffersFromDefaultForObject = true;
 				}
diff --git a/Engine/Source/Editor/SwarmInterface/Private/SwarmInterface.cpp b/Engine/Source/Editor/SwarmInterface/Private/SwarmInterface.cpp
--- a/Engine/Source/Editor/SwarmInterface/Private/SwarmInterface.cpp
+++ b/Engine/Source/Editor/SwarmInterface/Private/SwarmInterface.cpp
@@ -152,13 +152,13 @@ extern "C" DLLEXPORT void SwarmInterfaceLog(TVerbosityLevel Verbosity, const TCH
 	switch (Verbosity)
 	{
 	case VERBOSITY_Critical:
-		UE_LOG(LogSwarmInterface, Error, TEXT("%s"), Message);
+		UE_LOG_CLINKAGE(LogSwarmInterface, Error, TEXT("%s"), Message);
 		break;
 	case VERBOSITY_Complex:
-		UE_LOG(LogSwarmInterface, Warning, TEXT("%s"), Message);
+		UE_LOG_CLINKAGE(LogSwarmInterface, Warning, TEXT("%s"), Message);
 		break;
 	default:
-		UE_LOG(LogSwarmInterface, Log, TEXT("%s"), Message);
+		UE_LOG_CLINKAGE(LogSwarmInterface, Log, TEXT("%s"), Message);
 		break;
 	}
 }
diff --git a/Engine/Source/Editor/UnrealEd/Classes/CookOnTheSide/CookOnTheFlyServer.h b/Engine/Source/Editor/UnrealEd/Classes/CookOnTheSide/CookOnTheFlyServer.h
--- a/Engine/Source/Editor/UnrealEd/Classes/CookOnTheSide/CookOnTheFlyServer.h
+++ b/Engine/Source/Editor/UnrealEd/Classes/CookOnTheSide/CookOnTheFlyServer.h
@@ -471,6 +471,8 @@ public:
 	*/
 	bool IsCookByTheBookMode() const;
 
+	bool IsUsingShaderCodeLibrary() const;
+
 	bool IsUsingPackageStore() const;
 
 	/**
diff --git a/Engine/Source/Editor/UnrealEd/Private/Commandlets/CookCommandlet.cpp b/Engine/Source/Editor/UnrealEd/Private/Commandlets/CookCommandlet.cpp
--- a/Engine/Source/Editor/UnrealEd/Private/Commandlets/CookCommandlet.cpp
+++ b/Engine/Source/Editor/UnrealEd/Private/Commandlets/CookCommandlet.cpp
@@ -829,6 +829,7 @@ bool UCookCommandlet::CookByTheBook( const TArray<ITargetPlatform*>& Platforms,
 	CookOptions |= Switches.Contains(TEXT("NODEV")) ? ECookByTheBookOptions::NoDevContent : ECookByTheBookOptions::None;
 	CookOptions |= Switches.Contains(TEXT("FullLoadAndSave")) ? ECookByTheBookOptions::FullLoadAndSave : ECookByTheBookOptions::None;
 	CookOptions |= Switches.Contains(TEXT("PackageStore")) ? ECookByTheBookOptions::PackageStore : ECookByTheBookOptions::None;
+	CookOptions |= Switches.Contains(TEXT("NoGameAlwaysCook")) ? ECookByTheBookOptions::NoGameAlwaysCookPackages : ECookByTheBookOptions::None;
 
 	const ECookByTheBookOptions SinglePackageFlags = ECookByTheBookOptions::NoAlwaysCookMaps | ECookByTheBookOptions::NoDefaultMaps | ECookByTheBookOptions::NoGameAlwaysCookPackages | ECookByTheBookOptions::NoInputPackages | ECookByTheBookOptions::NoSlatePackages | ECookByTheBookOptions::DisableUnsolicitedPackages | ECookByTheBookOptions::ForceDisableSaveGlobalShaders;
 	CookOptions |= bCookSinglePackage ? SinglePackageFlags : ECookByTheBookOptions::None;
diff --git a/Engine/Source/Editor/UnrealEd/Private/Commandlets/ImportLocalizedDialogueCommandlet.cpp b/Engine/Source/Editor/UnrealEd/Private/Commandlets/ImportLocalizedDialogueCommandlet.cpp
--- a/Engine/Source/Editor/UnrealEd/Private/Commandlets/ImportLocalizedDialogueCommandlet.cpp
+++ b/Engine/Source/Editor/UnrealEd/Private/Commandlets/ImportLocalizedDialogueCommandlet.cpp
@@ -15,6 +15,7 @@
 #include "Interfaces/ITargetPlatform.h"
 #include "Interfaces/ITargetPlatformManagerModule.h"
 #include "AudioEditorModule.h"
+#include "AudioCompressionSettingsUtils.h"
 
 DEFINE_LOG_CATEGORY_STATIC(LogImportLocalizedDialogueCommandlet, Log, All);
 
@@ -513,7 +514,7 @@ USoundWave* UImportLocalizedDialogueCommandlet::ImportSoundWave(const FString& I
 			const TArray<ITargetPlatform*>& Platforms = TPM->GetActiveTargetPlatforms();
 			for (ITargetPlatform* Platform : Platforms)
 			{
-				SoundWave->GetCompressedData(Platform->GetWaveFormat(SoundWave), Platform->GetAudioCompressionSettings());
+				SoundWave->GetCompressedData(Platform->GetWaveFormat(SoundWave), FPlatformCompressionUtilities::GetCookOverrides(*Platform->IniPlatformName()));
 			}
 		}
 	}
diff --git a/Engine/Source/Editor/UnrealEd/Private/EditorBuildUtils.cpp b/Engine/Source/Editor/UnrealEd/Private/EditorBuildUtils.cpp
--- a/Engine/Source/Editor/UnrealEd/Private/EditorBuildUtils.cpp
+++ b/Engine/Source/Editor/UnrealEd/Private/EditorBuildUtils.cpp
@@ -1496,10 +1496,10 @@ void FMaterialOfflineCompilation::CopyPlatformSpecificStats()
 		return;
 	}
 
-	TMap<FName, FShader*> SrcShaders;
+	TMap<FHashedName, TShaderRef<FShader>> SrcShaders;
 	SrcShaderMap->GetShaderList(SrcShaders);
 
-	TMap<FName, FShader*> DstShaders;
+	TMap<FHashedName, TShaderRef<FShader>> DstShaders;
 	DstShaderMap->GetShaderList(DstShaders);
 
 	for (auto Pair : SrcShaders)
@@ -1507,11 +1507,8 @@ void FMaterialOfflineCompilation::CopyPlatformSpecificStats()
 		auto *DestinationShaderPtr = DstShaders.Find(Pair.Key);
 		if (DestinationShaderPtr != nullptr)
 		{
-			FShader *DestinationShader = *DestinationShaderPtr;
-			FShader *SourceShader = Pair.Value;
-
-			auto NumInstructions = SourceShader->GetNumInstructions();
-			DestinationShader->SetNumInstructions(NumInstructions);
+			auto NumInstructions = Pair.Value->GetNumInstructions();
+			(*DestinationShaderPtr)->SetNumInstructions(NumInstructions);
 		}
 	}
 }
diff --git a/Engine/Source/Editor/UnrealEd/Private/EditorViewportCommands.cpp b/Engine/Source/Editor/UnrealEd/Private/EditorViewportCommands.cpp
--- a/Engine/Source/Editor/UnrealEd/Private/EditorViewportCommands.cpp
+++ b/Engine/Source/Editor/UnrealEd/Private/EditorViewportCommands.cpp
@@ -313,20 +313,18 @@ static void AppendMaterialInfoToMenu(const UMaterialInterface* MaterialInterface
 	const FMaterialResource* Material = MaterialInterface->GetMaterialResource(FeatureLevel);
 	if (Material)
 	{
-		const TArray<TRefCountPtr<FMaterialUniformExpressionTexture> >& ExpressionsByType = Material->GetUniform2DTextureExpressions();
-		for (FMaterialUniformExpressionTexture* Expression : ExpressionsByType)
+		const FUniformExpressionSet& UniformExpressions = Material->GetUniformExpressions();
+		for (int32 i = 0; i < UniformExpressions.GetNumTextures(EMaterialTextureParameterType::Standard2D); ++i)
 		{
-			if (Expression)
-			{
-				UTexture* Texture = nullptr;
-				Expression->GetGameThreadTextureValue(MaterialInterface, *Material, Texture, true);
+			UTexture* Texture = nullptr;
+			UniformExpressions.GetGameThreadTextureValue(EMaterialTextureParameterType::Standard2D, i, MaterialInterface, *Material, Texture, true);
 
-				const UTexture2D* Texture2D = Cast<UTexture2D>(Texture);
-				if (Texture2D)
-				{
-					DataPerTextureIndex.FindOrAdd(Expression->GetTextureIndex()).AddUnique(FString::Printf(TEXT("%s.%s"), *MaterialInterface->GetName(), *Texture2D->GetName()));
-					DataPerTextureName.FindOrAdd(*Texture2D->GetName()).AddUnique(FString::Printf(TEXT("%s %d : %s"), *MenuName, Expression->GetTextureIndex(), *MaterialInterface->GetName()));
-				}
+			const UTexture2D* Texture2D = Cast<UTexture2D>(Texture);
+			if (Texture2D)
+			{
+				const FMaterialTextureParameterInfo& Parameter = UniformExpressions.GetTextureParameter(EMaterialTextureParameterType::Standard2D, i);
+				DataPerTextureIndex.FindOrAdd(Parameter.TextureIndex).AddUnique(FString::Printf(TEXT("%s.%s"), *MaterialInterface->GetName(), *Texture2D->GetName()));
+				DataPerTextureName.FindOrAdd(*Texture2D->GetName()).AddUnique(FString::Printf(TEXT("%s %d : %s"), *MenuName, Parameter.TextureIndex, *MaterialInterface->GetName()));
 			}
 		}
 	}
diff --git a/Engine/Source/Editor/UnrealEd/Private/Lightmass/LightmassRender.cpp b/Engine/Source/Editor/UnrealEd/Private/Lightmass/LightmassRender.cpp
--- a/Engine/Source/Editor/UnrealEd/Private/Lightmass/LightmassRender.cpp
+++ b/Engine/Source/Editor/UnrealEd/Private/Lightmass/LightmassRender.cpp
@@ -239,7 +239,7 @@ public:
 			PropertyToCompile = InPropertyToCompile;
 			Usage = InUsage;
 
-			MaterialInterface->AppendReferencedTextures(ReferencedTextures);
+			ReferencedTextures = MaterialInterface->GetReferencedTextures();
 
 			FMaterialResource* Resource = InMaterialInterface->GetMaterialResource(GMaxRHIFeatureLevel);
 			if (Resource)
@@ -247,9 +247,6 @@ public:
 				FMaterialShaderMapId ResourceId;
 				Resource->GetShaderMapId(GMaxRHIShaderPlatform, ResourceId);
 
-				FStaticParameterSet StaticParamSet;
-				Resource->GetStaticParameterSet(GMaxRHIShaderPlatform, StaticParamSet);
-
 				{
 					TArray<FShaderType*> ShaderTypes;
 					TArray<FVertexFactoryType*> VFTypes;
@@ -263,12 +260,12 @@ public:
 
 				// Override with a special usage so we won't re-use the shader map used by the material for rendering
 				ResourceId.Usage = GetShaderMapUsage();
-				CacheShaders(ResourceId, &StaticParamSet, GMaxRHIShaderPlatform);
+				CacheShaders(ResourceId, GMaxRHIShaderPlatform);
 			}
 		}
 	}
 
-	virtual const TArray<UObject*>& GetReferencedTextures() const override
+	virtual TArrayView<UObject* const> GetReferencedTextures() const override
 	{
 		return ReferencedTextures;
 	}
@@ -322,22 +319,22 @@ public:
 		}
 	}
 
-	virtual bool GetVectorValue(const FMaterialParameterInfo& ParameterInfo, FLinearColor* OutValue, const FMaterialRenderContext& Context) const override
+	virtual bool GetVectorValue(const FHashedMaterialParameterInfo& ParameterInfo, FLinearColor* OutValue, const FMaterialRenderContext& Context) const override
 	{
 		return MaterialInterface->GetRenderProxy()->GetVectorValue(ParameterInfo, OutValue, Context);
 	}
 
-	virtual bool GetScalarValue(const FMaterialParameterInfo& ParameterInfo, float* OutValue, const FMaterialRenderContext& Context) const override
+	virtual bool GetScalarValue(const FHashedMaterialParameterInfo& ParameterInfo, float* OutValue, const FMaterialRenderContext& Context) const override
 	{
 		return MaterialInterface->GetRenderProxy()->GetScalarValue(ParameterInfo, OutValue, Context);
 	}
 
-	virtual bool GetTextureValue(const FMaterialParameterInfo& ParameterInfo,const UTexture** OutValue, const FMaterialRenderContext& Context) const override
+	virtual bool GetTextureValue(const FHashedMaterialParameterInfo& ParameterInfo,const UTexture** OutValue, const FMaterialRenderContext& Context) const override
 	{
 		return MaterialInterface->GetRenderProxy()->GetTextureValue(ParameterInfo,OutValue,Context);
 	}
 
-	virtual bool GetTextureValue(const FMaterialParameterInfo& ParameterInfo, const URuntimeVirtualTexture** OutValue, const FMaterialRenderContext& Context) const override
+	virtual bool GetTextureValue(const FHashedMaterialParameterInfo& ParameterInfo, const URuntimeVirtualTexture** OutValue, const FMaterialRenderContext& Context) const override
 	{
 		return MaterialInterface->GetRenderProxy()->GetTextureValue(ParameterInfo, OutValue, Context);
 	}
diff --git a/Engine/Source/Editor/UnrealEd/Private/MeshPaintRendering.cpp b/Engine/Source/Editor/UnrealEd/Private/MeshPaintRendering.cpp
--- a/Engine/Source/Editor/UnrealEd/Private/MeshPaintRendering.cpp
+++ b/Engine/Source/Editor/UnrealEd/Private/MeshPaintRendering.cpp
@@ -42,21 +42,13 @@ namespace MeshPaintRendering
 			TransformParameter.Bind( Initializer.ParameterMap, TEXT( "c_Transform" ) );
 		}
 
-		virtual bool Serialize( FArchive& Ar ) override
-		{
-			bool bShaderHasOutdatedParameters = FGlobalShader::Serialize( Ar );
-			Ar << TransformParameter;
-			return bShaderHasOutdatedParameters;
-		}
-
 		void SetParameters(FRHICommandList& RHICmdList, const FMatrix& InTransform )
 		{
-			SetShaderValue(RHICmdList, GetVertexShader(), TransformParameter, InTransform );
+			SetShaderValue(RHICmdList, RHICmdList.GetBoundVertexShader(), TransformParameter, InTransform );
 		}
 
 	private:
-
-		FShaderParameter TransformParameter;
+		LAYOUT_FIELD(FShaderParameter, TransformParameter);
 	};
 
 
@@ -68,7 +60,6 @@ namespace MeshPaintRendering
 	class TMeshPaintPixelShader : public FGlobalShader
 	{
 		DECLARE_SHADER_TYPE( TMeshPaintPixelShader, Global );
-
 	public:
 
 		static bool ShouldCompilePermutation(const FGlobalShaderPermutationParameters& Parameters)
@@ -94,24 +85,9 @@ namespace MeshPaintRendering
 			GammaParameter.Bind( Initializer.ParameterMap, TEXT( "c_Gamma" ) );
 		}
 
-		virtual bool Serialize(FArchive& Ar) override
-		{
-			bool bShaderHasOutdatedParameters = FGlobalShader::Serialize(Ar);
-			Ar << CloneTextureParameter;
-			Ar << CloneTextureParameterSampler;
-			Ar << WorldToBrushMatrixParameter;
-			Ar << BrushMetricsParameter;
-			Ar << BrushStrengthParameter;
-			Ar << BrushColorParameter;
-			Ar << ChannelFlagsParameter;
-			Ar << GenerateMaskFlagParameter;
-			Ar << GammaParameter;
-			return bShaderHasOutdatedParameters;
-		}
-
 		void SetParameters(FRHICommandList& RHICmdList, const float InGamma, const FMeshPaintShaderParameters& InShaderParams )
 		{
-			FRHIPixelShader* ShaderRHI = GetPixelShader();
+			FRHIPixelShader* ShaderRHI = RHICmdList.GetBoundPixelShader();
 
 			SetTextureParameter(
 				RHICmdList, 
@@ -149,34 +125,32 @@ namespace MeshPaintRendering
 			SetShaderValue(RHICmdList, ShaderRHI, GammaParameter, InGamma );
 		}
 
-
 	private:
-
 		/** Texture that is a clone of the destination render target before we start drawing */
-		FShaderResourceParameter CloneTextureParameter;
-		FShaderResourceParameter CloneTextureParameterSampler;
+		LAYOUT_FIELD(FShaderResourceParameter, CloneTextureParameter);
+		LAYOUT_FIELD(FShaderResourceParameter, CloneTextureParameterSampler);
 
 		/** Brush -> World matrix */
-		FShaderParameter WorldToBrushMatrixParameter;
+		LAYOUT_FIELD(FShaderParameter, WorldToBrushMatrixParameter);
 
 		/** Brush metrics: x = radius, y = falloff range, z = depth, w = depth falloff range */
-		FShaderParameter BrushMetricsParameter;
+		LAYOUT_FIELD(FShaderParameter, BrushMetricsParameter);
 
 		/** Brush strength */
-		FShaderParameter BrushStrengthParameter;
+		LAYOUT_FIELD(FShaderParameter, BrushStrengthParameter);
 
 		/** Brush color */
-		FShaderParameter BrushColorParameter;
+		LAYOUT_FIELD(FShaderParameter, BrushColorParameter);
 
 		/** Flags that control paining individual channels: x = Red, y = Green, z = Blue, w = Alpha */
-		FShaderParameter ChannelFlagsParameter;
+		LAYOUT_FIELD(FShaderParameter, ChannelFlagsParameter);
 		
 		/** Flag to control brush mask generation or paint blending */
-		FShaderParameter GenerateMaskFlagParameter;
+		LAYOUT_FIELD(FShaderParameter, GenerateMaskFlagParameter);
 
 		/** Gamma */
 		// @todo MeshPaint: Remove this?
-		FShaderParameter GammaParameter;
+		LAYOUT_FIELD(FShaderParameter, GammaParameter);
 	};
 
 
@@ -205,21 +179,14 @@ namespace MeshPaintRendering
 			TransformParameter.Bind( Initializer.ParameterMap, TEXT( "c_Transform" ) );
 		}
 
-		virtual bool Serialize( FArchive& Ar ) override
-		{
-			bool bShaderHasOutdatedParameters = FShader::Serialize( Ar );
-			Ar << TransformParameter;
-			return bShaderHasOutdatedParameters;
-		}
-
 		void SetParameters(FRHICommandList& RHICmdList, const FMatrix& InTransform )
 		{
-			SetShaderValue(RHICmdList, GetVertexShader(), TransformParameter, InTransform );
+			SetShaderValue(RHICmdList, RHICmdList.GetBoundVertexShader(), TransformParameter, InTransform );
 		}
 
 	private:
 
-		FShaderParameter TransformParameter;
+		LAYOUT_FIELD(FShaderParameter, TransformParameter);
 	};
 
 
@@ -257,24 +224,9 @@ namespace MeshPaintRendering
 			GammaParameter.Bind( Initializer.ParameterMap, TEXT( "Gamma" ) );
 		}
 
-		virtual bool Serialize(FArchive& Ar) override
-		{
-			bool bShaderHasOutdatedParameters = FShader::Serialize(Ar);
-			Ar << Texture0Parameter;
-			Ar << Texture0ParameterSampler;
-			Ar << Texture1Parameter;
-			Ar << Texture1ParameterSampler;
-			Ar << Texture2Parameter;
-			Ar << Texture2ParameterSampler;
-			Ar << WidthPixelOffsetParameter;
-			Ar << HeightPixelOffsetParameter;
-			Ar << GammaParameter;
-			return bShaderHasOutdatedParameters;
-		}
-
 		void SetParameters(FRHICommandList& RHICmdList, const float InGamma, const FMeshPaintDilateShaderParameters& InShaderParams )
 		{
-			FRHIPixelShader* ShaderRHI = GetPixelShader();
+			FRHIPixelShader* ShaderRHI = RHICmdList.GetBoundPixelShader();
 
 			SetTextureParameter(
 				RHICmdList, 
@@ -309,26 +261,26 @@ namespace MeshPaintRendering
 	private:
 
 		/** Texture0 */
-		FShaderResourceParameter Texture0Parameter;
-		FShaderResourceParameter Texture0ParameterSampler;
+		LAYOUT_FIELD(FShaderResourceParameter, Texture0Parameter);
+		LAYOUT_FIELD(FShaderResourceParameter, Texture0ParameterSampler);
 
 		/** Texture1 */
-		FShaderResourceParameter Texture1Parameter;
-		FShaderResourceParameter Texture1ParameterSampler;
+		LAYOUT_FIELD(FShaderResourceParameter, Texture1Parameter);
+		LAYOUT_FIELD(FShaderResourceParameter, Texture1ParameterSampler);
 
 		/** Texture2 */
-		FShaderResourceParameter Texture2Parameter;
-		FShaderResourceParameter Texture2ParameterSampler;
+		LAYOUT_FIELD(FShaderResourceParameter, Texture2Parameter);
+		LAYOUT_FIELD(FShaderResourceParameter, Texture2ParameterSampler);
 
 		/** Pixel size width */
-		FShaderParameter WidthPixelOffsetParameter;
+		LAYOUT_FIELD(FShaderParameter, WidthPixelOffsetParameter);
 		
 		/** Pixel size height */
-		FShaderParameter HeightPixelOffsetParameter;
+		LAYOUT_FIELD(FShaderParameter, HeightPixelOffsetParameter);
 
 		/** Gamma */
 		// @todo MeshPaint: Remove this?
-		FShaderParameter GammaParameter;
+		LAYOUT_FIELD(FShaderParameter, GammaParameter);
 	};
 
 
@@ -362,8 +314,8 @@ namespace MeshPaintRendering
 		TShaderMapRef< TMeshPaintPixelShader > PixelShader(GetGlobalShaderMap(InFeatureLevel));
 
 		GraphicsPSOInit.BoundShaderState.VertexDeclarationRHI = GMeshPaintDilateVertexDeclaration.VertexDeclarationRHI;
-		GraphicsPSOInit.BoundShaderState.VertexShaderRHI = GETSAFERHISHADER_VERTEX(*VertexShader);
-		GraphicsPSOInit.BoundShaderState.PixelShaderRHI = GETSAFERHISHADER_PIXEL(*PixelShader);
+		GraphicsPSOInit.BoundShaderState.VertexShaderRHI = VertexShader.GetVertexShader();
+		GraphicsPSOInit.BoundShaderState.PixelShaderRHI = PixelShader.GetPixelShader();
 		GraphicsPSOInit.PrimitiveType = PT_TriangleList;
 
 		SetGraphicsPipelineState(RHICmdList, GraphicsPSOInit, EApplyRendertargetOption::ForceApply);
@@ -386,8 +338,8 @@ namespace MeshPaintRendering
 		TShaderMapRef< TMeshPaintDilatePixelShader > PixelShader(GetGlobalShaderMap(InFeatureLevel));
 
 		GraphicsPSOInit.BoundShaderState.VertexDeclarationRHI = GMeshPaintDilateVertexDeclaration.VertexDeclarationRHI;
-		GraphicsPSOInit.BoundShaderState.VertexShaderRHI = GETSAFERHISHADER_VERTEX(*VertexShader);
-		GraphicsPSOInit.BoundShaderState.PixelShaderRHI = GETSAFERHISHADER_PIXEL(*PixelShader);
+		GraphicsPSOInit.BoundShaderState.VertexShaderRHI = VertexShader.GetVertexShader();
+		GraphicsPSOInit.BoundShaderState.PixelShaderRHI = PixelShader.GetPixelShader();
 		GraphicsPSOInit.PrimitiveType = PT_TriangleList;
 
 		SetGraphicsPipelineState(RHICmdList, GraphicsPSOInit, EApplyRendertargetOption::ForceApply);
diff --git a/Engine/Source/Editor/UnrealEd/Private/NormalMapPreview.cpp b/Engine/Source/Editor/UnrealEd/Private/NormalMapPreview.cpp
--- a/Engine/Source/Editor/UnrealEd/Private/NormalMapPreview.cpp
+++ b/Engine/Source/Editor/UnrealEd/Private/NormalMapPreview.cpp
@@ -49,26 +49,14 @@ public:
 	 */
 	void SetParameters(FRHICommandList& RHICmdList, const FTexture* NormalMapTexture)
 	{
-		FRHIPixelShader* PixelShaderRHI = GetPixelShader();
+		FRHIPixelShader* PixelShaderRHI = RHICmdList.GetBoundPixelShader();
 		SetTextureParameter(RHICmdList, PixelShaderRHI,Texture,TextureSampler,NormalMapTexture);
 	}
 
-	/**
-	 * Serialization.
-	 * @param Ar - The archive with which to serialize.
-	 */
-	virtual bool Serialize(FArchive& Ar) override
-	{
-		bool bShaderHasOutdatedParameters = FGlobalShader::Serialize(Ar);
-		Ar << Texture;
-		Ar << TextureSampler;
-		return bShaderHasOutdatedParameters;
-	}
-
 private:
 	/** The texture to sample. */
-	FShaderResourceParameter Texture;
-	FShaderResourceParameter TextureSampler;
+	LAYOUT_FIELD(FShaderResourceParameter, Texture);
+	LAYOUT_FIELD(FShaderResourceParameter, TextureSampler);
 };
 IMPLEMENT_SHADER_TYPE(,FSimpleElementNormalMapPS,TEXT("/Engine/Private/SimpleElementNormalMapPixelShader.usf"),TEXT("Main"),SF_Pixel);
 
@@ -86,8 +74,8 @@ void FNormalMapBatchedElementParameters::BindShaders(
 	TShaderMapRef<FSimpleElementNormalMapPS> PixelShader(GetGlobalShaderMap(InFeatureLevel));
 
 	GraphicsPSOInit.BoundShaderState.VertexDeclarationRHI = GSimpleElementVertexDeclaration.VertexDeclarationRHI;
-	GraphicsPSOInit.BoundShaderState.VertexShaderRHI = GETSAFERHISHADER_VERTEX(*VertexShader);
-	GraphicsPSOInit.BoundShaderState.PixelShaderRHI = GETSAFERHISHADER_PIXEL(*PixelShader);
+	GraphicsPSOInit.BoundShaderState.VertexShaderRHI = VertexShader.GetVertexShader();
+	GraphicsPSOInit.BoundShaderState.PixelShaderRHI = PixelShader.GetPixelShader();
 	GraphicsPSOInit.PrimitiveType = PT_TriangleList;
 
 	GraphicsPSOInit.BlendState = TStaticBlendState<>::GetRHI();
