Fix for warning after upgrading to PVS-Studio 7.0.4:

V1028: Possible overflow. Consider casting operands of the 'j + IABOffsetIndex' operator to the 'uint64' type, not the result.

#rnx

diff --git a/Engine/Source/Developer/Apple/MetalShaderFormat/Private/MetalDerivedData.cpp b/Engine/Source/Developer/Apple/MetalShaderFormat/Private/MetalDerivedData.cpp
--- a/Engine/Source/Developer/Apple/MetalShaderFormat/Private/MetalDerivedData.cpp
+++ b/Engine/Source/Developer/Apple/MetalShaderFormat/Private/MetalDerivedData.cpp
@@ -1209,8 +1209,8 @@ bool FMetalShaderOutputCooker::Build(TArray<uint8>& OutData)
 						if (Entry->UniformBufferName == TableNames[j])
 						{
 							Entry->SetIndex = j;
-							BufferIndices &= ~(1ull << (uint64)(j + IABOffsetIndex));
-							TextureIndices &= ~(1ull << uint64(j + IABOffsetIndex));
+							BufferIndices &= ~(1ull << ((uint64)j + IABOffsetIndex));
+							TextureIndices &= ~(1ull << ((uint64)j + IABOffsetIndex));
 							break;
 						}
 					}
