Fix issues caught by address sanitizer.

#rnx

diff --git a/Engine/Source/Developer/Apple/MetalShaderFormat/Private/MetalDerivedData.cpp b/Engine/Source/Developer/Apple/MetalShaderFormat/Private/MetalDerivedData.cpp
--- a/Engine/Source/Developer/Apple/MetalShaderFormat/Private/MetalDerivedData.cpp
+++ b/Engine/Source/Developer/Apple/MetalShaderFormat/Private/MetalDerivedData.cpp
@@ -373,7 +373,7 @@ FString FMetalShaderOutputCooker::GetPluginSpecificCacheKeySuffix() const
 			Flags |= (1ull << uint64(Flag));
 		}
 
-		CachedOutputName = FString::Printf(TEXT("%s-%s_%s-%u_%hu_%llu_%hhu_%d_%s_%s"), *Input.ShaderFormat.GetPlainNameString(), *Input.EntryPointName, *Hash.ToString(), Len, FormatVers, Flags, VersionEnum, IABTier, *GUIDHash.ToString(), *Standard);
+		CachedOutputName = FString::Printf(TEXT("%s-%s_%s-%u_%hu_%llu_%hu_%d_%s_%s"), *Input.ShaderFormat.GetPlainNameString(), *Input.EntryPointName, *Hash.ToString(), Len, FormatVers, Flags, VersionEnum, IABTier, *GUIDHash.ToString(), *Standard);
 	}
 
 	return CachedOutputName;
@@ -2075,6 +2075,8 @@ bool FMetalShaderOutputCooker::Build(TArray<uint8>& OutData)
 		Result = (Results.hasError) ? 0 : 1;
 		if (!Results.hasError && Results.target)
 		{
+			const std::string ResultsTargetDataAsString((const char*)Results.target->Data(), Results.target->Size());
+
 			MetaData += TEXT("// Compiled by ShaderConductor\n");
 			if (INPString.Len())
 			{
@@ -2112,9 +2114,7 @@ bool FMetalShaderOutputCooker::Build(TArray<uint8>& OutData)
 			{
 				MetaData += FString::Printf(TEXT("// @NumThreads: %s\n"), *WKGString);
 			}
-			
-			
-			if(strstr((const char*)Results.target->Data(), "spvBufferSizeConstants"))
+			if(FCStringAnsi::Strstr(ResultsTargetDataAsString.c_str(), "spvBufferSizeConstants"))
 			{
 				MetaData += FString::Printf(TEXT("// @SideTable: spvBufferSizeConstants(%d)\n"), SideTableIndex);
 			}
@@ -2142,7 +2142,7 @@ bool FMetalShaderOutputCooker::Build(TArray<uint8>& OutData)
 			}
 			
 			MetalSource = TCHAR_TO_UTF8(*MetaData);
-			MetalSource += std::string((const char*)Results.target->Data(), Results.target->Size());
+			MetalSource += ResultsTargetDataAsString;
 			
 			// Tessellation vertex & hull shaders must always use FMA
 			if (Options.enableFMAPass)
