Fixed unusable error report when Metal compiler is not available.

#rnx

diff --git a/Engine/Source/Developer/Apple/MetalShaderFormat/Private/MetalShaderCompiler.cpp b/Engine/Source/Developer/Apple/MetalShaderFormat/Private/MetalShaderCompiler.cpp
--- a/Engine/Source/Developer/Apple/MetalShaderFormat/Private/MetalShaderCompiler.cpp
+++ b/Engine/Source/Developer/Apple/MetalShaderFormat/Private/MetalShaderCompiler.cpp
@@ -618,9 +618,9 @@ bool ExtractCompilerInfo(EShaderPlatform ShaderPlatform, FString* OutVersion, FS
 	
 	{
 		// Fetch the directories where the binaries live and where metal_stdlib lives
-		// xcrun -sdk <sdk> metal -print-search-dirs
+		// $ xcrun -sdk <sdk> metal --print-search-dirs
 		// For example (Xcode 11.1):
-		// xcrun -sdk macosx metal -print-search-dirs
+		// $ xcrun -sdk macosx metal --print-search-dirs
 		// programs: =/Applications/Xcode.app/Contents/Developer/Toolchains/XcodeDefault.xctoolchain/usr/metal/macos/bin
 		// libraries: =/Applications/Xcode.app/Contents/Developer/Toolchains/XcodeDefault.xctoolchain/usr/metal/macos/lib/clang/902.9
 		FString OutputString;
@@ -665,7 +665,7 @@ static bool SingleCompilerSetup(EAppleSDKType SDK)
 	EShaderPlatform ShaderPlatform = AppleSDKToBaseShaderPlatform(SDK);
 	if (!ExtractCompilerInfo(ShaderPlatform, &Version, &BinaryDirectory, &LibraryDirectory))
 	{
-		UE_LOG(LogMetalShaderCompiler, Fatal, TEXT("BAD"));
+		UE_LOG(LogMetalShaderCompiler, Fatal, TEXT("Failed to extract Metal compiler search directories"));
 		return false;
 	}
 	
