Fix scanning Metal search directories when multiple paths are concatenated by colons (such as $PATH environment variable).

#rnx

diff --git a/Engine/Source/Developer/Apple/MetalShaderFormat/Private/MetalShaderCompiler.cpp b/Engine/Source/Developer/Apple/MetalShaderFormat/Private/MetalShaderCompiler.cpp
--- a/Engine/Source/Developer/Apple/MetalShaderFormat/Private/MetalShaderCompiler.cpp
+++ b/Engine/Source/Developer/Apple/MetalShaderFormat/Private/MetalShaderCompiler.cpp
@@ -575,6 +575,61 @@ FString GetXcodePath()
 #endif
 }
 
+// PathPrefix should be "programs:" or "libraries:" and both followed by "=DIR" where DIR is the path to extract.
+// This function also handles the case when multiple paths are concatenated via colons (like it's the case with the $PATH environment variable).
+static bool ExtractXcodeCompilerPath(const FString& InPathInfo, const TCHAR* PathPrefix, FString& OutPath, const TCHAR* RequiredFilename)
+{
+	if (InPathInfo.Contains(PathPrefix))
+	{
+		// Scan output directory. Note that it might contain multiple paths separated by colons (like the $PATH environment variable)
+		int32 IndexStart = InPathInfo.Find(TEXT("="));
+		
+		if (InPathInfo.Find(TEXT(":"), ESearchCase::IgnoreCase, ESearchDir::FromStart, IndexStart + 1) != INDEX_NONE)
+		{
+			// Warn about multiple paths, as this may indicate an issue with Xcode setup
+			UE_LOG(LogMetalShaderCompiler, Warning,
+				   TEXT("Multiple paths returned by Metal compiler info. This may indicate an issue with the installation of Xcode: %s"),
+				   *InPathInfo.RightChop(IndexStart + 1));
+			
+			// Find directory in concatenated path list that contains the required file, either "metal" or "include/metal/metal_stdlib"
+			for (int32 IndexEnd = 0; IndexStart != INDEX_NONE; IndexStart = IndexEnd)
+			{
+				// Skip the current "=" or ":" character
+				++IndexStart;
+				IndexEnd = InPathInfo.Find(TEXT(":"), ESearchCase::IgnoreCase, ESearchDir::FromStart, IndexStart);
+				
+				// Extract install directory DIR from first substring of "programs: =DIR:FURTHER_DIRS"
+				if (IndexEnd == INDEX_NONE)
+				{
+					OutPath = InPathInfo.RightChop(IndexStart);
+				}
+				else
+				{
+					OutPath = InPathInfo.Mid(IndexStart, IndexEnd - IndexStart);
+				}
+				
+				// Check if required file exists in this directory
+				if (FPaths::FileExists(OutPath / RequiredFilename))
+				{
+					// Found required file, stop scanning for paths
+					return true;
+				}
+			}
+		}
+		else
+		{
+			// Extract install directory DIR from right side of "programs: =DIR"
+			OutPath = InPathInfo.RightChop(IndexStart + 1);
+			
+			// Check if required file exists in this directory
+			return FPaths::FileExists(OutPath / RequiredFilename);
+		}
+	}
+	
+	// Compiler path not found
+	return false;
+}
+
 bool ExtractCompilerInfo(EShaderPlatform ShaderPlatform, FString* OutVersion, FString* OutInstalledDirectory, FString* OutLibDirectory)
 {
 	{
@@ -599,6 +654,7 @@ bool ExtractCompilerInfo(EShaderPlatform ShaderPlatform, FString* OutVersion, FS
 			if(VersionStart != -1)
 			{
 				// this should be something in the form of metalfe-XXX.X.XX
+				check(OutVersion);
 				*OutVersion = OutputString.RightChop(VersionStart+1);
 				int32 End = OutVersion->Find(TEXT(")"));
 				*OutVersion = OutVersion->Left(End);
@@ -630,19 +686,16 @@ bool ExtractCompilerInfo(EShaderPlatform ShaderPlatform, FString* OutVersion, FS
 			// split into lines and look for the output
 			TArray<FString> Lines;
 			OutputString.ParseIntoArrayLines(Lines, true);
-
-			if(Lines[0].Contains(TEXT("programs:")))
-			{
-				int32 Index = Lines[0].Find(TEXT("="));
-				*OutInstalledDirectory = Lines[0].RightChop(Index+1);
-			}
 			
-			if(Lines[1].Contains(TEXT("libraries:")))
+			// Extract directory where the "metal" executable lives
+			check(OutInstalledDirectory);
+			ExtractXcodeCompilerPath(Lines[0], TEXT("programs:"), *OutInstalledDirectory, TEXT("metal"));
+			
+			// Extract directory where the "include/metal/metal_stdlib" header file lives, and append the additional reltive path
+			check(OutLibDirectory);
+			if (ExtractXcodeCompilerPath(Lines[1], TEXT("libraries:"), *OutLibDirectory, TEXT("include/metal/metal_stdlib")))
 			{
-				int32 Index = Lines[1].Find(TEXT("="));
-				*OutLibDirectory = Lines[1].RightChop(Index+1);
-				
-				// ends up pointing to the clang version base. we want the metal headers.
+				// Ends up pointing to the clang version base. we want the metal headers.
 				*OutLibDirectory /= TEXT("include/metal");
 			}
 		}
