Removed warning about multiple search paths in Metal compiler info (blocks build farm).

#rnx

diff --git a/Engine/Source/Developer/Apple/MetalShaderFormat/Private/MetalShaderCompiler.cpp b/Engine/Source/Developer/Apple/MetalShaderFormat/Private/MetalShaderCompiler.cpp
--- a/Engine/Source/Developer/Apple/MetalShaderFormat/Private/MetalShaderCompiler.cpp
+++ b/Engine/Source/Developer/Apple/MetalShaderFormat/Private/MetalShaderCompiler.cpp
@@ -586,11 +586,6 @@ static bool ExtractXcodeCompilerPath(const FString& InPathInfo, const TCHAR* Pat
 		
 		if (InPathInfo.Find(TEXT(":"), ESearchCase::IgnoreCase, ESearchDir::FromStart, IndexStart + 1) != INDEX_NONE)
 		{
-			// Warn about multiple paths, as this may indicate an issue with Xcode setup
-			UE_LOG(LogMetalShaderCompiler, Warning,
-				   TEXT("Multiple paths returned by Metal compiler info. This may indicate an issue with the installation of Xcode: %s"),
-				   *InPathInfo.RightChop(IndexStart + 1));
-			
 			// Find directory in concatenated path list that contains the required file, either "metal" or "include/metal/metal_stdlib"
 			for (int32 IndexEnd = 0; IndexStart != INDEX_NONE; IndexStart = IndexEnd)
 			{
