Solve Remote Shadercompile

diff --git a/Engine/Source/Developer/Apple/MetalShaderFormat/Private/MetalShaderCompiler.cpp b/Engine/Source/Developer/Apple/MetalShaderFormat/Private/MetalShaderCompiler.cpp
--- a/Engine/Source/Developer/Apple/MetalShaderFormat/Private/MetalShaderCompiler.cpp
+++ b/Engine/Source/Developer/Apple/MetalShaderFormat/Private/MetalShaderCompiler.cpp
@@ -604,7 +604,7 @@ static bool ExtractXcodeCompilerPath(const FString& InPathInfo, const TCHAR* Pat
 				}
 				
 				// Check if required file exists in this directory
-				if (FPaths::FileExists(OutPath / RequiredFilename))
+				if (RemoteFileExists(OutPath / RequiredFilename))
 				{
 					// Found required file, stop scanning for paths
 					return true;
@@ -617,7 +617,7 @@ static bool ExtractXcodeCompilerPath(const FString& InPathInfo, const TCHAR* Pat
 			OutPath = InPathInfo.RightChop(IndexStart + 1);
 			
 			// Check if required file exists in this directory
-			return FPaths::FileExists(OutPath / RequiredFilename);
+			return RemoteFileExists(OutPath / RequiredFilename);
 		}
 	}
 	
