Add materials that have WPO to the FMA pass logic as only vertex ouput position can be invariant - this goes with the DXC change to add invariant for metal >= 2.1 to the vetex output position.

diff --git a/Engine/Source/Developer/Apple/MetalShaderFormat/Private/MetalDerivedData.cpp b/Engine/Source/Developer/Apple/MetalShaderFormat/Private/MetalDerivedData.cpp
--- a/Engine/Source/Developer/Apple/MetalShaderFormat/Private/MetalDerivedData.cpp
+++ b/Engine/Source/Developer/Apple/MetalShaderFormat/Private/MetalDerivedData.cpp
@@ -582,7 +582,7 @@ bool FMetalShaderOutputCooker::Build(TArray<uint8>& OutData)
 		Options.enableDebugInfo = false;
 		Options.enable16bitTypes = false;
 		Options.disableOptimizations = false;
-		Options.enableFMAPass = (VersionEnum == 2 || VersionEnum == 3);
+		Options.enableFMAPass = (VersionEnum == 2 || VersionEnum == 3 || bForceInvariance);
 		
         ShaderConductor::Compiler::SourceDesc SourceDesc;
 		
