Adopt FShaderCodeName::Key in the Metal shader compiler.

- Follow-on to CL 14816162

#rnx

diff --git a/Engine/Source/Developer/Apple/MetalShaderFormat/Private/MetalShaderCompiler.cpp b/Engine/Source/Developer/Apple/MetalShaderFormat/Private/MetalShaderCompiler.cpp
--- a/Engine/Source/Developer/Apple/MetalShaderFormat/Private/MetalShaderCompiler.cpp
+++ b/Engine/Source/Developer/Apple/MetalShaderFormat/Private/MetalShaderCompiler.cpp
@@ -817,8 +817,8 @@ void BuildMetalShaderOutput(
 		Ar << Header;
 		Ar.Serialize((void*)USFSource, SourceLen + 1 - (USFSource - InShaderSource));
 		
-		// store data we can pickup later with ShaderCode.FindOptionalData('n'), could be removed for shipping
-		ShaderOutput.ShaderCode.AddOptionalData('n', TCHAR_TO_UTF8(*ShaderInput.GenerateShaderName()));
+		// store data we can pickup later with ShaderCode.FindOptionalData(FShaderCodeName::Key), could be removed for shipping
+		ShaderOutput.ShaderCode.AddOptionalData(FShaderCodeName::Key, TCHAR_TO_UTF8(*ShaderInput.GenerateShaderName()));
 
 		if (ShaderInput.ExtraSettings.bExtractShaderSource)
 		{
@@ -1027,8 +1027,8 @@ void BuildMetalShaderOutput(
 			
 			if (ShaderInput.Environment.CompilerFlags.Contains(CFLAG_KeepDebugInfo))
 			{
-				// store data we can pickup later with ShaderCode.FindOptionalData('n'), could be removed for shipping
-				ShaderOutput.ShaderCode.AddOptionalData('n', TCHAR_TO_UTF8(*ShaderInput.GenerateShaderName()));
+				// store data we can pickup later with ShaderCode.FindOptionalData(FShaderCodeName::Key), could be removed for shipping
+				ShaderOutput.ShaderCode.AddOptionalData(FShaderCodeName::Key, TCHAR_TO_UTF8(*ShaderInput.GenerateShaderName()));
 				if (DebugCode.CompressedData.Num() == 0)
 				{
 					ShaderOutput.ShaderCode.AddOptionalData('c', TCHAR_TO_UTF8(*MetalCode));
@@ -1540,10 +1540,10 @@ void CompileShader_Metal(const FShaderCompilerInput& _Input,FShaderCompilerOutpu
 			TArray<uint8> SourceCode;
 			SourceCode.Append(SourceCodePtr, ShaderCode.GetActualShaderCodeSize() - CodeOffset);
 				
-			// store data we can pickup later with ShaderCode.FindOptionalData('n'), could be removed for shipping
+			// store data we can pickup later with ShaderCode.FindOptionalData(FShaderCodeName::Key), could be removed for shipping
 			ANSICHAR const* Text = ShaderCode.FindOptionalData('c');
 			ANSICHAR const* Path = ShaderCode.FindOptionalData('p');
-			ANSICHAR const* Name = ShaderCode.FindOptionalData('n');
+			ANSICHAR const* Name = ShaderCode.FindOptionalData(FShaderCodeName::Key);
 				
 			int32 ObjectSize = 0;
 			uint8 const* Object = ShaderCode.FindOptionalDataAndSize('o', ObjectSize);
@@ -1568,7 +1568,7 @@ void CompileShader_Metal(const FShaderCompilerInput& _Input,FShaderCompilerOutpu
 				
 			if (Name)
 			{
-				Output.ShaderCode.AddOptionalData('n', TCHAR_TO_UTF8(*Input.GenerateShaderName()));
+				Output.ShaderCode.AddOptionalData(FShaderCodeName::Key, TCHAR_TO_UTF8(*Input.GenerateShaderName()));
 			}
 			if (Path)
 			{
