Fix shader compile error compiling TPageTableUpdateVS<false>for IOS from PC

diff --git a/Engine/Source/Developer/Apple/MetalShaderFormat/Private/MetalDerivedData.cpp b/Engine/Source/Developer/Apple/MetalShaderFormat/Private/MetalDerivedData.cpp
--- a/Engine/Source/Developer/Apple/MetalShaderFormat/Private/MetalDerivedData.cpp
+++ b/Engine/Source/Developer/Apple/MetalShaderFormat/Private/MetalDerivedData.cpp
@@ -1682,6 +1682,7 @@ bool FMetalShaderOutputCooker::Build(TArray<uint8>& OutData)
 				TargetDesc.Language = CrossCompiler::EShaderConductorLanguage::Metal_iOS;
 				TargetDesc.CompileFlags.SetDefine(TEXT("use_framebuffer_fetch_subpasses"), 1);
 				TargetDesc.CompileFlags.SetDefine(TEXT("emulate_cube_array"), 1);
+				TargetDesc.CompileFlags.SetDefine(TEXT("ios_support_base_vertex_instance"), 1);
 				break;
 			}
 
