Provide metal tools output to caller upon failure, dramatically simplifying debugging.

#rnx

diff --git a/Engine/Source/Developer/Apple/MetalShaderFormat/Private/MetalShaderCompiler.cpp b/Engine/Source/Developer/Apple/MetalShaderFormat/Private/MetalShaderCompiler.cpp
--- a/Engine/Source/Developer/Apple/MetalShaderFormat/Private/MetalShaderCompiler.cpp
+++ b/Engine/Source/Developer/Apple/MetalShaderFormat/Private/MetalShaderCompiler.cpp
@@ -957,6 +957,7 @@ void BuildMetalShaderOutput(
 			Job.SourceCRC = SourceCRC;
 			Job.bRetainObjectFile = ShaderInput.Environment.CompilerFlags.Contains(CFLAG_Archive);
 			Job.bCompileAsPCH = false;
+			Job.ReturnCode = 0;
 
 			FMetalShaderBytecodeCooker* BytecodeCooker = new FMetalShaderBytecodeCooker(Job);
 			
diff --git a/Engine/Source/Developer/Apple/MetalShaderFormat/Private/MetalShaderFormat.cpp b/Engine/Source/Developer/Apple/MetalShaderFormat/Private/MetalShaderFormat.cpp
--- a/Engine/Source/Developer/Apple/MetalShaderFormat/Private/MetalShaderFormat.cpp
+++ b/Engine/Source/Developer/Apple/MetalShaderFormat/Private/MetalShaderFormat.cpp
@@ -854,9 +854,9 @@ bool FMetalCompilerToolchain::CompileMetalShader(FMetalShaderBytecodeJob& Job, F
 		FString MetalParams = FString::Printf(TEXT("%s %s %s %s -Wno-null-character -fbracket-depth=1024 %s %s %s %s -o %s"), *Job.MinOSVersion, *Job.DebugInfo, *Job.MathMode, TEXT("-c"), *Job.Standard, *Job.Defines, *IncludeArgs, *LocalInputMetalFilePath, *LocalOutputMetalAIRFilePath);
 		bool bSuccess = this->ExecMetalFrontend(SDK, *MetalParams, &Job.ReturnCode, &Job.Results, &Job.Errors);
 
-		if (!bSuccess)
+		if (!bSuccess || (Job.ReturnCode != 0))
 		{
-			Job.Message = FString::Printf(TEXT("[%s] Failed to compile %s to bytecode %s, code: %d, output: %s %s"), *LocalInputMetalFilePath, *LocalOutputMetalAIRFilePath, Job.ReturnCode, *Job.Results, *Job.Errors);
+			Job.Message = FString::Printf(TEXT("Failed to compile %s to bytecode %s, code: %d, output: %s %s"), *LocalInputMetalFilePath, *LocalOutputMetalAIRFilePath, Job.ReturnCode, *Job.Results, *Job.Errors);
 			return false;
 		}
 	}
@@ -867,7 +867,7 @@ bool FMetalCompilerToolchain::CompileMetalShader(FMetalShaderBytecodeJob& Job, F
 
 		FString MetalLibParams = FString::Printf(TEXT("-o %s %s"), *LocalOutputMetalLibFilePath, *LocalOutputMetalAIRFilePath);
 		bool bSuccess = this->ExecMetalLib(SDK, *MetalLibParams, &Job.ReturnCode, &Job.Results, &Job.Errors);
-		if (!bSuccess)
+		if (!bSuccess || (Job.ReturnCode != 0))
 		{
 			Job.Message = FString::Printf(TEXT("Failed to package %s into %s, code: %d, output: %s %s"), *LocalOutputMetalAIRFilePath, *LocalOutputMetalLibFilePath, Job.ReturnCode, *Job.Results, *Job.Errors);
 			return false;
