Explicitly specify the Target when determining the metal version.

The target info reported when running `xcrun --sdk <platform> metal -v' is always the default for the running OS, and may vary slightly between point releases. Specifying the target reduces the chances of mismatched shader version when using FASTBuild to compile shaders.

NOTE: In order for the shader version to match, the FASTBuild workers must (currently) have the same version Xcode installed, but do not necessarily need to be running the same OS version.

#[fyi] richard.wallis
#rnx

diff --git a/Engine/Source/Developer/Apple/MetalShaderFormat/Private/MetalShaderFormat.cpp b/Engine/Source/Developer/Apple/MetalShaderFormat/Private/MetalShaderFormat.cpp
--- a/Engine/Source/Developer/Apple/MetalShaderFormat/Private/MetalShaderFormat.cpp
+++ b/Engine/Source/Developer/Apple/MetalShaderFormat/Private/MetalShaderFormat.cpp
@@ -628,7 +628,7 @@ FMetalCompilerToolchain::EMetalToolchainStatus FMetalCompilerToolchain::FetchCom
 		FString StdOut;
 		// metal -v writes its output to stderr
 		// But the underlying (windows) implementation of CreateProc puts everything into one pipe, which is written to StdOut.
-		bool bResult = this->ExecMetalFrontend(AppleSDKMac, TEXT("-v"), &ReturnCode, &StdOut, &StdOut);
+		bool bResult = this->ExecMetalFrontend(AppleSDKMac, TEXT("-v --target=air64-apple-darwin18.7.0"), &ReturnCode, &StdOut, &StdOut);
 		check(bResult);
 		
 		Result = ParseCompilerVersionAndTarget(StdOut, this->MetalCompilerVersionString[AppleSDKMac], this->MetalCompilerVersion[AppleSDKMac], this->MetalTargetVersion[AppleSDKMac]);
@@ -643,7 +643,7 @@ FMetalCompilerToolchain::EMetalToolchainStatus FMetalCompilerToolchain::FetchCom
 		int32 ReturnCode = 0;
 		FString StdOut;
 		// metal -v writes its output to stderr
-		bool bResult = this->ExecMetalFrontend(AppleSDKMobile, TEXT("-v"), &ReturnCode, &StdOut, &StdOut);
+		bool bResult = this->ExecMetalFrontend(AppleSDKMobile, TEXT("-v --target=air64-apple-darwin18.7.0"), &ReturnCode, &StdOut, &StdOut);
 		check(bResult);
 		
 		Result = ParseCompilerVersionAndTarget(StdOut, this->MetalCompilerVersionString[AppleSDKMobile], this->MetalCompilerVersion[AppleSDKMobile], this->MetalTargetVersion[AppleSDKMobile]);
