Remove references to UE4 in Metal backend.

#rnx

diff --git a/Engine/Source/Developer/Apple/MetalShaderFormat/Private/MetalDerivedData.cpp b/Engine/Source/Developer/Apple/MetalShaderFormat/Private/MetalDerivedData.cpp
--- a/Engine/Source/Developer/Apple/MetalShaderFormat/Private/MetalDerivedData.cpp
+++ b/Engine/Source/Developer/Apple/MetalShaderFormat/Private/MetalDerivedData.cpp
@@ -1199,13 +1199,13 @@ bool FMetalShaderOutputCooker::Build(TArray<uint8>& OutData)
 				std::string FMADefine = std::string("\n"
 										"template<typename T>\n"
 										"static inline __attribute__((always_inline))\n"
-										"T ue4_cross(T x, T y)\n"
+										"T ue_cross(T x, T y)\n"
 										"{\n"
 										"    metal::float3 fx = metal::float3(x);\n"
 										"    metal::float3 fy = metal::float3(y);\n"
 										"    return T(metal::fma(fx[1], fy[2], -metal::fma(fy[1], fx[2], 0.0)), metal::fma(fx[2], fy[0], -metal::fma(fy[2], fx[0], 0.0)), metal::fma(fx[0], fy[1], -metal::fma(fy[0], fx[1], 0.0)));\n"
 										"}\n"
-										"#define cross ue4_cross\n\n"
+										"#define cross ue_cross\n\n"
 										"using namespace metal;"
 										);
 				
diff --git a/Engine/Source/Developer/Apple/MetalShaderFormat/Private/MetalShaderCompiler.cpp b/Engine/Source/Developer/Apple/MetalShaderFormat/Private/MetalShaderCompiler.cpp
--- a/Engine/Source/Developer/Apple/MetalShaderFormat/Private/MetalShaderCompiler.cpp
+++ b/Engine/Source/Developer/Apple/MetalShaderFormat/Private/MetalShaderCompiler.cpp
@@ -1049,7 +1049,7 @@ void CompileShader_Metal(const FShaderCompilerInput& _Input,FShaderCompilerOutpu
 				Output.bSucceeded = false;
 				FShaderCompilerError* NewError = new(Output.Errors) FShaderCompilerError();
 				NewError->StrippedErrorMessage = FString::Printf(
-															 TEXT("Metal %s is no longer supported in UE4 for macOS."),
+															 TEXT("Metal %s is no longer supported in UE for macOS."),
 															 *StandardVersion
 															 );
 				return;
@@ -1064,7 +1064,7 @@ void CompileShader_Metal(const FShaderCompilerInput& _Input,FShaderCompilerOutpu
 			Output.bSucceeded = false;
 			FShaderCompilerError* NewError = new(Output.Errors) FShaderCompilerError();
 			NewError->StrippedErrorMessage = FString::Printf(
-														 TEXT("Metal %s is no longer supported in UE4."),
+														 TEXT("Metal %s is no longer supported in UE."),
 														 *StandardVersion
 														 );
 			return;
@@ -1079,7 +1079,7 @@ void CompileShader_Metal(const FShaderCompilerInput& _Input,FShaderCompilerOutpu
 			Output.bSucceeded = false;
 			FShaderCompilerError* NewError = new(Output.Errors) FShaderCompilerError();
 			NewError->StrippedErrorMessage = FString::Printf(
-															 TEXT("Metal %s is no longer supported in UE4."),
+															 TEXT("Metal %s is no longer supported in UE."),
 															 *StandardVersion
 															 );
 			return;
