- Fix some unnecessary string copies when calling FPaths::FileExists

diff --git a/Engine/Plugins/Developer/PluginUtils/Source/PluginUtils/Private/PluginUtils.cpp b/Engine/Plugins/Developer/PluginUtils/Source/PluginUtils/Private/PluginUtils.cpp
--- a/Engine/Plugins/Developer/PluginUtils/Source/PluginUtils/Private/PluginUtils.cpp
+++ b/Engine/Plugins/Developer/PluginUtils/Source/PluginUtils/Private/PluginUtils.cpp
@@ -646,7 +646,7 @@ bool FPluginUtils::ValidateNewPluginNameAndLocation(const FString& PluginName, c
 		{
 			const FString PluginFilePath = FPluginUtils::GetPluginFilePath(PluginLocation, PluginName);
 
-			if (!PluginFilePath.IsEmpty() && FPaths::FileExists(*PluginFilePath))
+			if (!PluginFilePath.IsEmpty() && FPaths::FileExists(PluginFilePath))
 			{
 				if (FailReason)
 				{
diff --git a/Engine/Plugins/Experimental/GameFeatures/Source/GameFeatures/Private/GameFeaturePluginStateMachine.cpp b/Engine/Plugins/Experimental/GameFeatures/Source/GameFeatures/Private/GameFeaturePluginStateMachine.cpp
--- a/Engine/Plugins/Experimental/GameFeatures/Source/GameFeatures/Private/GameFeaturePluginStateMachine.cpp
+++ b/Engine/Plugins/Experimental/GameFeatures/Source/GameFeatures/Private/GameFeaturePluginStateMachine.cpp
@@ -223,7 +223,7 @@ struct FGameFeaturePluginState_CheckingStatus : public FGameFeaturePluginState
 
 		if (StateProperties.GetPluginProtocol() == EGameFeaturePluginProtocol::File)
 		{
-			bIsAvailable = FPaths::FileExists(*StateProperties.PluginInstalledFilename);
+			bIsAvailable = FPaths::FileExists(StateProperties.PluginInstalledFilename);
 		}
 		else if (StateProperties.GetPluginProtocol() == EGameFeaturePluginProtocol::InstallBundle)
 		{
diff --git a/Engine/Plugins/Online/OnlineSubsystemSteam/Source/Private/OnlineSubsystemSteam.cpp b/Engine/Plugins/Online/OnlineSubsystemSteam/Source/Private/OnlineSubsystemSteam.cpp
--- a/Engine/Plugins/Online/OnlineSubsystemSteam/Source/Private/OnlineSubsystemSteam.cpp
+++ b/Engine/Plugins/Online/OnlineSubsystemSteam/Source/Private/OnlineSubsystemSteam.cpp
@@ -174,7 +174,7 @@ static void DeleteSteamAppIdFromDisk()
 	const FString SteamAppIdFilename = GetSteamAppIdFilename();
 	// Turn off sandbox temporarily to make sure file is where it's always expected
 	FScopeSandboxContext ScopedSandbox(false);
-	if (FPaths::FileExists(*SteamAppIdFilename))
+	if (FPaths::FileExists(SteamAppIdFilename))
 	{
 		bool bSuccessfullyDeleted = IFileManager::Get().Delete(*SteamAppIdFilename);
 	}
diff --git a/Engine/Plugins/Runtime/GoogleVR/GoogleVRHMD/Source/GoogleVRHMD/Public/GoogleVRAdbUtils.h b/Engine/Plugins/Runtime/GoogleVR/GoogleVRHMD/Source/GoogleVRHMD/Public/GoogleVRAdbUtils.h
--- a/Engine/Plugins/Runtime/GoogleVR/GoogleVRHMD/Source/GoogleVRHMD/Public/GoogleVRAdbUtils.h
+++ b/Engine/Plugins/Runtime/GoogleVR/GoogleVRHMD/Source/GoogleVRHMD/Public/GoogleVRAdbUtils.h
@@ -54,7 +54,7 @@ inline void GetAdbPath(FString& OutAdbPath)
 #endif
 
 		// if it doesn't exist then just clear the path as we might set it later
-		if (!FPaths::FileExists(*OutAdbPath))
+		if (!FPaths::FileExists(OutAdbPath))
 		{
 			OutAdbPath.Empty();
 		}
diff --git a/Engine/Plugins/Runtime/ResonanceAudio/Source/ResonanceAudio/Private/ResonanceAudioCommon.cpp b/Engine/Plugins/Runtime/ResonanceAudio/Source/ResonanceAudio/Private/ResonanceAudioCommon.cpp
--- a/Engine/Plugins/Runtime/ResonanceAudio/Source/ResonanceAudio/Private/ResonanceAudioCommon.cpp
+++ b/Engine/Plugins/Runtime/ResonanceAudio/Source/ResonanceAudio/Private/ResonanceAudioCommon.cpp
@@ -42,7 +42,7 @@ namespace ResonanceAudio
 
 		UE_LOG(LogResonanceAudio, Log, TEXT("Attempting to load %s"), *DynamicLibraryToLoad);
 
-		if (FPaths::FileExists(*DynamicLibraryToLoad))
+		if (FPaths::FileExists(DynamicLibraryToLoad))
 		{
 			DynamicLibraryHandle = FPlatformProcess::GetDllHandle(*DynamicLibraryToLoad);
 		}
diff --git a/Engine/Plugins/Tests/EditorTests/Source/EditorTests/Private/UnrealEd/EditorPerformanceTests.cpp b/Engine/Plugins/Tests/EditorTests/Source/EditorTests/Private/UnrealEd/EditorPerformanceTests.cpp
--- a/Engine/Plugins/Tests/EditorTests/Source/EditorTests/Private/UnrealEd/EditorPerformanceTests.cpp
+++ b/Engine/Plugins/Tests/EditorTests/Source/EditorTests/Private/UnrealEd/EditorPerformanceTests.cpp
@@ -65,7 +65,7 @@ void EditorPerfDump(EditorPerfCaptureParameters& EditorPerfStats)
 	//Get the map load time (in seconds) from the text file that is created when the load map latent command is ran.
 	EditorPerfStats.MapLoadTime = 0;
 	FString MapLoadTimeFileLocation = FPaths::Combine(*DataFileLocation, TEXT("RAWMapLoadTime.txt"));
-	if (FPaths::FileExists(*MapLoadTimeFileLocation))
+	if (FPaths::FileExists(MapLoadTimeFileLocation))
 	{
 		TArray<FString> SavedMapLoadTimes;
 		FAutomationEditorCommonUtils::CreateArrayFromFile(MapLoadTimeFileLocation, SavedMapLoadTimes);
@@ -109,7 +109,7 @@ void EditorPerfDump(EditorPerfCaptureParameters& EditorPerfStats)
 	FTimespan TestRunDuration = (EditorPerfStats.TimeStamp.Last().GetTicks() - EditorPerfStats.TimeStamp[0].GetTicks()) + ETimespan::TicksPerSecond;
 
 	//The performance csv file will be created if it didn't exist prior to the start of this test.
-	if (!FPaths::FileExists(*PerfCSVFilePath))
+	if (!FPaths::FileExists(PerfCSVFilePath))
 	{
 		FArchive* FinalCSVArchive = IFileManager::Get().CreateFileWriter(*PerfCSVFilePath);
 		if ( FAutomationEditorCommonUtils::IsArchiveWriteable(PerfCSVFilePath, FinalCSVArchive))
diff --git a/Engine/Source/Developer/Android/AndroidDeviceDetection/Private/AndroidDeviceDetectionModule.cpp b/Engine/Source/Developer/Android/AndroidDeviceDetection/Private/AndroidDeviceDetectionModule.cpp
--- a/Engine/Source/Developer/Android/AndroidDeviceDetection/Private/AndroidDeviceDetectionModule.cpp
+++ b/Engine/Source/Developer/Android/AndroidDeviceDetection/Private/AndroidDeviceDetectionModule.cpp
@@ -131,7 +131,7 @@ private:
 			OutStdErr = &DefaultError;
 		}
 
-		if (FPaths::FileExists(*ADBPath))
+		if (FPaths::FileExists(ADBPath))
 		{
 			FPlatformProcess::ExecProcess(*ADBPath, *CommandLine, &ReturnCode, OutStdOut, OutStdErr);
 
@@ -688,7 +688,7 @@ public:
 			ADBPath = FPaths::Combine(*AndroidDirectory, SDKRelativeExePath);
 
 			// if it doesn't exist then just clear the path as we might set it later
-			if (!FPaths::FileExists(*ADBPath))
+			if (!FPaths::FileExists(ADBPath))
 			{
 				ADBPath.Empty();
 			}
diff --git a/Engine/Source/Developer/Android/AndroidPlatformEditor/Private/AndroidLicenseDialog.cpp b/Engine/Source/Developer/Android/AndroidPlatformEditor/Private/AndroidLicenseDialog.cpp
--- a/Engine/Source/Developer/Android/AndroidPlatformEditor/Private/AndroidLicenseDialog.cpp
+++ b/Engine/Source/Developer/Android/AndroidPlatformEditor/Private/AndroidLicenseDialog.cpp
@@ -164,7 +164,7 @@ static FString GetLicensePath()
 	IAndroidDeviceDetection* DeviceDetection = AndroidDeviceDetection.GetAndroidDeviceDetection();
 	FString ADBPath = DeviceDetection->GetADBPath();
 
-	if (!FPaths::FileExists(*ADBPath))
+	if (!FPaths::FileExists(ADBPath))
 	{
 		return TEXT("");
 	}
diff --git a/Engine/Source/Developer/Android/AndroidTargetPlatform/Private/AndroidTargetPlatform.cpp b/Engine/Source/Developer/Android/AndroidTargetPlatform/Private/AndroidTargetPlatform.cpp
--- a/Engine/Source/Developer/Android/AndroidTargetPlatform/Private/AndroidTargetPlatform.cpp
+++ b/Engine/Source/Developer/Android/AndroidTargetPlatform/Private/AndroidTargetPlatform.cpp
@@ -58,7 +58,7 @@ static FString GetLicensePath()
 	IAndroidDeviceDetection* DeviceDetection = AndroidDeviceDetection.GetAndroidDeviceDetection();
 	FString ADBPath = DeviceDetection->GetADBPath();
 
-	if (!FPaths::FileExists(*ADBPath))
+	if (!FPaths::FileExists(ADBPath))
 	{
 		return TEXT("");
 	}
diff --git a/Engine/Source/Developer/Apple/MetalShaderFormat/Private/MetalShaderCompiler.cpp b/Engine/Source/Developer/Apple/MetalShaderFormat/Private/MetalShaderCompiler.cpp
--- a/Engine/Source/Developer/Apple/MetalShaderFormat/Private/MetalShaderCompiler.cpp
+++ b/Engine/Source/Developer/Apple/MetalShaderFormat/Private/MetalShaderCompiler.cpp
@@ -676,12 +676,12 @@ void BuildMetalShaderOutput(
 					UE_LOG(LogMetalShaderCompiler, Fatal, TEXT("Failed to write Metal shader out to %s\nShaderText:\n%s"), *SaveFile, *MetalCode);
 				}
 				bSuccess = IFileManager::Get().Move(*MetalFileName, *SaveFile, false, false, true, true);
-				if (!bSuccess && !FPaths::FileExists(*MetalFileName))
+				if (!bSuccess && !FPaths::FileExists(MetalFileName))
 				{
 					UE_LOG(LogMetalShaderCompiler, Fatal, TEXT("Failed to move %s to %s"), *SaveFile, *MetalFileName);
 				}
 
-				if (FPaths::FileExists(*SaveFile))
+				if (FPaths::FileExists(SaveFile))
 				{
 					IFileManager::Get().Delete(*SaveFile);
 				}
@@ -1593,7 +1593,7 @@ bool FinalizeLibrary_Metal(FName const& Format, FString const& WorkingDir, FStri
 
 			FFileHelper::SaveStringToFile(M_Script, *LocalScriptFilePath);
 
-			if (!FPaths::FileExists(*LocalScriptFilePath))
+			if (!FPaths::FileExists(LocalScriptFilePath))
 			{
 				UE_LOG(LogMetalShaderCompiler, Error, TEXT("Failed to create metal-ar .M script at %s"), *LocalScriptFilePath);
 				return false;
@@ -1601,7 +1601,7 @@ bool FinalizeLibrary_Metal(FName const& Format, FString const& WorkingDir, FStri
 
 			FPaths::MakePlatformFilename(LocalScriptFilePath);
 			bool bSuccess = Toolchain->ExecMetalAr(SDK, *LocalScriptFilePath, &ReturnCode, &Results, &Errors);
-			bArchiveFileValid = FPaths::FileExists(*LocalArchivePath);
+			bArchiveFileValid = FPaths::FileExists(LocalArchivePath);
 					
 			if (ReturnCode != 0 || !bArchiveFileValid)
 			{
diff --git a/Engine/Source/Developer/Apple/MetalShaderFormat/Private/MetalShaderFormat.cpp b/Engine/Source/Developer/Apple/MetalShaderFormat/Private/MetalShaderFormat.cpp
--- a/Engine/Source/Developer/Apple/MetalShaderFormat/Private/MetalShaderFormat.cpp
+++ b/Engine/Source/Developer/Apple/MetalShaderFormat/Private/MetalShaderFormat.cpp
@@ -695,7 +695,7 @@ FMetalCompilerToolchain::EMetalToolchainStatus FMetalCompilerToolchain::DoMacNat
 			MetalFrontendBinaryCommand[AppleSDKMac] = ToolchainBase / TEXT("macos") / TEXT("bin") / MetalFrontendBinary;
 			MetalFrontendBinaryCommand[AppleSDKMobile] = ToolchainBase / TEXT("ios") / TEXT("bin") / MetalFrontendBinary;
 
-			const bool bIsFrontendPresent = FPaths::FileExists(*MetalFrontendBinaryCommand[AppleSDKMac]) && FPaths::FileExists(*MetalFrontendBinaryCommand[AppleSDKMobile]);
+			const bool bIsFrontendPresent = FPaths::FileExists(MetalFrontendBinaryCommand[AppleSDKMac]) && FPaths::FileExists(MetalFrontendBinaryCommand[AppleSDKMobile]);
 			if (!bIsFrontendPresent)
 			{
 				UE_LOG(LogMetalCompilerSetup, Warning, TEXT("Missing Metal frontend in %s."), *ToolchainBase);
@@ -708,10 +708,10 @@ FMetalCompilerToolchain::EMetalToolchainStatus FMetalCompilerToolchain::DoMacNat
 			MetalLibBinaryCommand[AppleSDKMac] = ToolchainBase / TEXT("macos") / TEXT("bin") / MetalLibraryBinary;
 			MetalLibBinaryCommand[AppleSDKMobile] = ToolchainBase / TEXT("ios") / TEXT("bin") / MetalLibraryBinary;
 
-			if (!FPaths::FileExists(*MetalArBinaryCommand[AppleSDKMac]) ||
-				!FPaths::FileExists(*MetalArBinaryCommand[AppleSDKMobile]) ||
-				!FPaths::FileExists(*MetalLibBinaryCommand[AppleSDKMac]) ||
-				!FPaths::FileExists(*MetalLibBinaryCommand[AppleSDKMobile]))
+			if (!FPaths::FileExists(MetalArBinaryCommand[AppleSDKMac]) ||
+				!FPaths::FileExists(MetalArBinaryCommand[AppleSDKMobile]) ||
+				!FPaths::FileExists(MetalLibBinaryCommand[AppleSDKMac]) ||
+				!FPaths::FileExists(MetalLibBinaryCommand[AppleSDKMobile]))
 			{
 				UE_LOG(LogMetalCompilerSetup, Warning, TEXT("Missing toolchain binaries in %s."), *ToolchainBase);
 				return EMetalToolchainStatus::ToolchainNotFound;
@@ -757,7 +757,7 @@ FMetalCompilerToolchain::EMetalToolchainStatus FMetalCompilerToolchain::DoWindow
 	MetalFrontendBinaryCommand[AppleSDKMac] = ToolchainBase / TEXT("macos") / TEXT("bin") / MetalFrontendBinary;
 	MetalFrontendBinaryCommand[AppleSDKMobile] = ToolchainBase / TEXT("ios") / TEXT("bin") / MetalFrontendBinary;
 
-	bool bUseLocalMetalToolchain = FPaths::FileExists(*MetalFrontendBinaryCommand[AppleSDKMac]) && FPaths::FileExists(*MetalFrontendBinaryCommand[AppleSDKMobile]);
+	bool bUseLocalMetalToolchain = FPaths::FileExists(MetalFrontendBinaryCommand[AppleSDKMac]) && FPaths::FileExists(MetalFrontendBinaryCommand[AppleSDKMobile]);
 	if (!bUseLocalMetalToolchain)
 	{
 #if CHECK_METAL_COMPILER_TOOLCHAIN_SETUP
@@ -773,10 +773,10 @@ FMetalCompilerToolchain::EMetalToolchainStatus FMetalCompilerToolchain::DoWindow
 	MetalLibBinaryCommand[AppleSDKMac] = ToolchainBase / TEXT("macos") / TEXT("bin") / MetalLibraryBinary;
 	MetalLibBinaryCommand[AppleSDKMobile] = ToolchainBase / TEXT("ios") / TEXT("bin") / MetalLibraryBinary;
 
-	if (!FPaths::FileExists(*MetalArBinaryCommand[AppleSDKMac]) ||
-		!FPaths::FileExists(*MetalArBinaryCommand[AppleSDKMobile]) ||
-		!FPaths::FileExists(*MetalLibBinaryCommand[AppleSDKMac]) ||
-		!FPaths::FileExists(*MetalLibBinaryCommand[AppleSDKMobile]))
+	if (!FPaths::FileExists(MetalArBinaryCommand[AppleSDKMac]) ||
+		!FPaths::FileExists(MetalArBinaryCommand[AppleSDKMobile]) ||
+		!FPaths::FileExists(MetalLibBinaryCommand[AppleSDKMac]) ||
+		!FPaths::FileExists(MetalLibBinaryCommand[AppleSDKMobile]))
 	{
 #if CHECK_METAL_COMPILER_TOOLCHAIN_SETUP
 		UE_LOG(LogMetalCompilerSetup, Warning, TEXT("Missing toolchain binaries."))
diff --git a/Engine/Source/Developer/DesktopPlatform/Private/DesktopPlatformBase.cpp b/Engine/Source/Developer/DesktopPlatform/Private/DesktopPlatformBase.cpp
--- a/Engine/Source/Developer/DesktopPlatform/Private/DesktopPlatformBase.cpp
+++ b/Engine/Source/Developer/DesktopPlatform/Private/DesktopPlatformBase.cpp
@@ -887,7 +887,7 @@ const TArray<FTargetInfo>& FDesktopPlatformBase::GetTargetsForProject(const FStr
 	// Run UBT to update the list of targets. Try to run it without building first.
 	FString Output;
 	int32 ReturnCode = 0;
-	if (!FPaths::FileExists(*GetUnrealBuildToolExecutableFilename(FPaths::RootDir())) || !const_cast<FDesktopPlatformBase*>(this)->InvokeUnrealBuildToolSync(Arguments, *GLog, true, ReturnCode, Output) || ReturnCode != 0)
+	if (!FPaths::FileExists(GetUnrealBuildToolExecutableFilename(FPaths::RootDir())) || !const_cast<FDesktopPlatformBase*>(this)->InvokeUnrealBuildToolSync(Arguments, *GLog, true, ReturnCode, Output) || ReturnCode != 0)
 	{
 		// If that failed, try to build and run again. Build machines should always have an up-to-date copy, and shouldn't build anything without being told to do so.
 		if (!GIsBuildMachine)
diff --git a/Engine/Source/Developer/DesktopPlatform/Private/Linux/DesktopPlatformLinux.cpp b/Engine/Source/Developer/DesktopPlatform/Private/Linux/DesktopPlatformLinux.cpp
--- a/Engine/Source/Developer/DesktopPlatform/Private/Linux/DesktopPlatformLinux.cpp
+++ b/Engine/Source/Developer/DesktopPlatform/Private/Linux/DesktopPlatformLinux.cpp
@@ -315,7 +315,7 @@ static bool CompareAndCheckDesktopFile(const TCHAR* DesktopFileName, const TCHAR
 		DesktopFileExecPath = Matcher.GetCaptureGroup(1);
 	}
 
-	if (DesktopFileExecPath.Compare("bash") != 0 && !FPaths::FileExists(*DesktopFileExecPath))
+	if (DesktopFileExecPath.Compare("bash") != 0 && !FPaths::FileExists(DesktopFileExecPath))
 	{
 		return false;
 	}
diff --git a/Engine/Source/Developer/LauncherServices/Private/Profiles/LauncherProfileManager.cpp b/Engine/Source/Developer/LauncherServices/Private/Profiles/LauncherProfileManager.cpp
--- a/Engine/Source/Developer/LauncherServices/Private/Profiles/LauncherProfileManager.cpp
+++ b/Engine/Source/Developer/LauncherServices/Private/Profiles/LauncherProfileManager.cpp
@@ -13,8 +13,8 @@
 LAUNCHERSERVICES_API bool HasPromotedTarget(const TCHAR* BaseDir, const TCHAR* TargetName, const TCHAR* Platform, EBuildConfiguration Configuration, const TCHAR* Architecture)
 {
 	// Get the path to the receipt, and check it exists
-	FString ReceiptPath = FTargetReceipt::GetDefaultPath(BaseDir, TargetName, Platform, Configuration, Architecture);
-	if (!FPaths::FileExists(*ReceiptPath))
+	const FString ReceiptPath = FTargetReceipt::GetDefaultPath(BaseDir, TargetName, Platform, Configuration, Architecture);
+	if (!FPaths::FileExists(ReceiptPath))
 	{
 		UE_LOG(LogLauncherProfile, Log, TEXT("Unable to use promoted target - %s does not exist."), *ReceiptPath);
 		return false;
diff --git a/Engine/Source/Developer/PakFileUtilities/Private/PakFileUtilities.cpp b/Engine/Source/Developer/PakFileUtilities/Private/PakFileUtilities.cpp
--- a/Engine/Source/Developer/PakFileUtilities/Private/PakFileUtilities.cpp
+++ b/Engine/Source/Developer/PakFileUtilities/Private/PakFileUtilities.cpp
@@ -5111,9 +5111,9 @@ bool ExecuteUnrealPak(const TCHAR* CmdLine)
 		FString InputDir = FPaths::GetPath(*NonOptionArguments[0]);
 
 		TArray<FString> InputPakFiles;
-		if (FPaths::FileExists(*NonOptionArguments[0]))
+		if (FPaths::FileExists(NonOptionArguments[0]))
 		{
-			InputPakFiles.Add(*NonOptionArguments[0]);
+			InputPakFiles.Add(NonOptionArguments[0]);
 		}
 		else
 		{
diff --git a/Engine/Source/Editor/GameProjectGeneration/Private/SProjectBrowser.cpp b/Engine/Source/Editor/GameProjectGeneration/Private/SProjectBrowser.cpp
--- a/Engine/Source/Editor/GameProjectGeneration/Private/SProjectBrowser.cpp
+++ b/Engine/Source/Editor/GameProjectGeneration/Private/SProjectBrowser.cpp
@@ -574,11 +574,11 @@ static TSharedPtr<FSlateDynamicImageBrush> GetThumbnailForProject(const FString&
 	const FString ThumbnailPNGFile = FPaths::GetBaseFilename(ProjectFilename, false) + TEXT(".png");
 	const FString AutoScreenShotPNGFile = FPaths::Combine(*FPaths::GetPath(ProjectFilename), TEXT("Saved"), TEXT("AutoScreenshot.png"));
 	FString PNGFileToUse;
-	if (FPaths::FileExists(*ThumbnailPNGFile))
+	if (FPaths::FileExists(ThumbnailPNGFile))
 	{
 		PNGFileToUse = ThumbnailPNGFile;
 	}
-	else if (FPaths::FileExists(*AutoScreenShotPNGFile))
+	else if (FPaths::FileExists(AutoScreenShotPNGFile))
 	{
 		PNGFileToUse = AutoScreenShotPNGFile;
 	}
diff --git a/Engine/Source/Editor/UnrealEd/Private/Tests/AutomationEditorCommon.cpp b/Engine/Source/Editor/UnrealEd/Private/Tests/AutomationEditorCommon.cpp
--- a/Engine/Source/Editor/UnrealEd/Private/Tests/AutomationEditorCommon.cpp
+++ b/Engine/Source/Editor/UnrealEd/Private/Tests/AutomationEditorCommon.cpp
@@ -400,7 +400,7 @@ void FAutomationEditorCommonUtils::CreateArrayFromFile(const FString& InFileLoca
 {
 	FString RawData;
 
-	if (FPaths::FileExists(*InFileLocation))
+	if (FPaths::FileExists(InFileLocation))
 	{
 		UE_LOG(LogEditorAutomationTests, VeryVerbose, TEXT("Loading and parsing the data from '%s' into an array."), *InFileLocation);
 		FFileHelper::LoadFileToString(RawData, *InFileLocation);
diff --git a/Engine/Source/Runtime/Core/Private/Misc/ConfigManifest.cpp b/Engine/Source/Runtime/Core/Private/Misc/ConfigManifest.cpp
--- a/Engine/Source/Runtime/Core/Private/Misc/ConfigManifest.cpp
+++ b/Engine/Source/Runtime/Core/Private/Misc/ConfigManifest.cpp
@@ -149,9 +149,9 @@ void MigrateToAgnosticIni(const TCHAR* SrcIniName, const TCHAR* DstIniName)
 	const FString OldIni = ProjectSpecificIniPath(SrcIniName);
 	const FString NewIni = ProjectAgnosticIniPath(DstIniName);
 
-	if (FPaths::FileExists(*OldIni))
+	if (FPaths::FileExists(OldIni))
 	{
-		if (!FPaths::FileExists(*NewIni))
+		if (!FPaths::FileExists(NewIni))
 		{
 			IFileManager::Get().Move(*NewIni, *OldIni);
 		}
diff --git a/Engine/Source/Runtime/Core/Private/Modules/ModuleManager.cpp b/Engine/Source/Runtime/Core/Private/Modules/ModuleManager.cpp
--- a/Engine/Source/Runtime/Core/Private/Modules/ModuleManager.cpp
+++ b/Engine/Source/Runtime/Core/Private/Modules/ModuleManager.cpp
@@ -222,7 +222,7 @@ bool FModuleManager::IsModuleUpToDate(const FName InModuleName) const
 
 	for (const TPair<FName, FString>& Pair : ModulePathMap)
 	{
-		if (!FPaths::FileExists(*Pair.Value))
+		if (!FPaths::FileExists(Pair.Value))
 		{
 			return false;
 		}
diff --git a/Engine/Source/Runtime/Core/Private/Windows/WindowsPlatformProcess.cpp b/Engine/Source/Runtime/Core/Private/Windows/WindowsPlatformProcess.cpp
--- a/Engine/Source/Runtime/Core/Private/Windows/WindowsPlatformProcess.cpp
+++ b/Engine/Source/Runtime/Core/Private/Windows/WindowsPlatformProcess.cpp
@@ -1895,7 +1895,7 @@ void *FWindowsPlatformProcess::LoadLibraryWithSearchPaths(const FString& FileNam
 {
 	// Make sure the initial module exists. If we can't find it from the path we're given, it's probably a system dll.
 	FString FullFileName = FileName;
-	if (FPaths::FileExists(*FullFileName))
+	if (FPaths::FileExists(FullFileName))
 	{
 		// Convert it to a full path, since LoadLibrary will try to resolve it against the executable directory (which may not be the same as the working dir)
 		FullFileName = FPaths::ConvertRelativePathToFull(FullFileName);
diff --git a/Engine/Source/Runtime/Engine/Private/AssetManager.cpp b/Engine/Source/Runtime/Engine/Private/AssetManager.cpp
--- a/Engine/Source/Runtime/Engine/Private/AssetManager.cpp
+++ b/Engine/Source/Runtime/Engine/Private/AssetManager.cpp
@@ -2318,8 +2318,8 @@ bool UAssetManager::OnAssetRegistryAvailableAfterInitialization(FName InName, FA
 		
 		{
 			TArray<uint8> Bytes;
-			FString Filename = FPaths::ProjectDir() / (TEXT("AssetRegistry") + InName.ToString()) + TEXT(".bin");
-			if (FPaths::FileExists(*Filename) && FFileHelper::LoadFileToArray(Bytes, *Filename))
+			const FString Filename = FPaths::ProjectDir() / (TEXT("AssetRegistry") + InName.ToString()) + TEXT(".bin");
+			if (FPaths::FileExists(Filename) && FFileHelper::LoadFileToArray(Bytes, *Filename))
 			{
 				bLoaded = true;
 				FMemoryReader Ar(Bytes);
diff --git a/Engine/Source/Runtime/Engine/Private/URL.cpp b/Engine/Source/Runtime/Engine/Private/URL.cpp
--- a/Engine/Source/Runtime/Engine/Private/URL.cpp
+++ b/Engine/Source/Runtime/Engine/Private/URL.cpp
@@ -463,7 +463,7 @@ FURL::FURL( FURL* Base, const TCHAR* TextURL, ETravelType Type )
 			FText MapNameError;
 			bool bFoundMap = false;
 
-			if (FPaths::FileExists(*URLStr) && FPackageName::TryConvertFilenameToLongPackageName(URLStr, MapFullName))
+			if (FPaths::FileExists(URLStr) && FPackageName::TryConvertFilenameToLongPackageName(URLStr, MapFullName))
 			{
 				Map = MapFullName;
 				bFoundMap = true;
diff --git a/Engine/Source/Runtime/Online/SSL/Private/SslCertificateManager.cpp b/Engine/Source/Runtime/Online/SSL/Private/SslCertificateManager.cpp
--- a/Engine/Source/Runtime/Online/SSL/Private/SslCertificateManager.cpp
+++ b/Engine/Source/Runtime/Online/SSL/Private/SslCertificateManager.cpp
@@ -318,7 +318,7 @@ void FSslCertificateManager::BuildRootCertificateArray()
 	FString OverrideCertificateBundlePath;
 	if (GConfig->GetString(TEXT("SSL"), TEXT("OverrideCertificateBundlePath"), OverrideCertificateBundlePath, GEngineIni) && OverrideCertificateBundlePath.Len() > 0)
 	{
-		if (FPaths::FileExists(*(OverrideCertificateBundlePath)))
+		if (FPaths::FileExists(OverrideCertificateBundlePath))
 		{
 			CertificateBundlePath = OverrideCertificateBundlePath;
 		}
@@ -328,17 +328,21 @@ void FSslCertificateManager::BuildRootCertificateArray()
 	if (CertificateBundlePath.IsEmpty())
 	{
 		const FString PerPlatformBundlePath = FString::Printf(TEXT("Certificates/%s/cacert.pem"), ANSI_TO_TCHAR(FPlatformProperties::IniPlatformName()));
-		if (FPaths::FileExists(*(FPaths::ProjectContentDir() + PerPlatformBundlePath)))
-		{
-			CertificateBundlePath = FPaths::ProjectContentDir() + PerPlatformBundlePath;
-		}
-		else if (FPaths::FileExists(*(FPaths::ProjectContentDir() + TEXT("Certificates/cacert.pem"))))
+
+		const FString SearchPaths[]
 		{
-			CertificateBundlePath = FPaths::ProjectContentDir() + TEXT("Certificates/cacert.pem");
-		}
-		else if (FPaths::FileExists(*(FPaths::EngineContentDir() + TEXT("Certificates/ThirdParty/cacert.pem"))))
+			FPaths::ProjectContentDir() + PerPlatformBundlePath,
+			FPaths::ProjectContentDir() + TEXT("Certificates/cacert.pem"),
+			FPaths::EngineContentDir() + TEXT("Certificates/ThirdParty/cacert.pem")
+		};
+
+		for (const FString& SearchPath : SearchPaths)
 		{
-			CertificateBundlePath = FPaths::EngineContentDir() + TEXT("Certificates/ThirdParty/cacert.pem");
+			if (FPaths::FileExists(SearchPath))
+			{
+				CertificateBundlePath = SearchPath;
+				break;
+			}
 		}
 	}
 
