Fix-up metal compiler version parsing to accommodate changes to Xcode 13 SDK.

#rnx

diff --git a/Engine/Source/Developer/Apple/MetalShaderFormat/Private/MetalShaderFormat.cpp b/Engine/Source/Developer/Apple/MetalShaderFormat/Private/MetalShaderFormat.cpp
--- a/Engine/Source/Developer/Apple/MetalShaderFormat/Private/MetalShaderFormat.cpp
+++ b/Engine/Source/Developer/Apple/MetalShaderFormat/Private/MetalShaderFormat.cpp
@@ -384,10 +384,11 @@ static FMetalCompilerToolchain::EMetalToolchainStatus ParseCompilerVersionAndTar
 
 		int32 Major = 0, Minor = 0, Patch = 0;
 		int32 NumResults = 0;
+		char AppleToolName[PATH_MAX] = { '\0' };
 #if !PLATFORM_WINDOWS
-		NumResults = sscanf(TCHAR_TO_ANSI(*Version), "Apple LLVM version %d.%d.%d", &Major, &Minor, &Patch);
+		NumResults = sscanf(TCHAR_TO_ANSI(*Version), "Apple %s version %d.%d.%d", AppleToolName, &Major, &Minor, &Patch);
 #else
-		NumResults = swscanf_s(*Version, TEXT("Apple LLVM version %d.%d.%d"), &Major, &Minor, &Patch);
+		NumResults = swscanf_s(*Version, TEXT("Apple %s version %d.%d.%d"), AppleToolName, &Major, &Minor, &Patch);
 #endif
 		PackedVersionNumber.Major = Major;
 		PackedVersionNumber.Minor = Minor;
@@ -797,7 +798,7 @@ bool FMetalCompilerToolchain::ExecMetalFrontend(EAppleSDKType SDK, const TCHAR*
 #if PLATFORM_MAC
 	if (this->MetalFrontendBinaryCommand[SDK].IsEmpty())
 	{
-		FString BuiltParams = FString::Printf(TEXT("-sdk %s %s %s"), *SDKToString(SDK), *this->MetalFrontendBinary, Parameters);
+		FString BuiltParams = FString::Printf(TEXT("--sdk %s %s %s"), *SDKToString(SDK), *this->MetalFrontendBinary, Parameters);
 		return ExecGenericCommand(*XcrunPath, *BuiltParams, OutReturnCode, OutStdOut, OutStdErr);
 	}
 	else
@@ -811,7 +812,7 @@ bool FMetalCompilerToolchain::ExecMetalLib(EAppleSDKType SDK, const TCHAR* Param
 #if PLATFORM_MAC
 	if (this->MetalLibBinaryCommand[SDK].IsEmpty())
 	{
-		FString BuiltParams = FString::Printf(TEXT("-sdk %s %s %s"), *SDKToString(SDK), *this->MetalLibraryBinary, Parameters);
+		FString BuiltParams = FString::Printf(TEXT("--sdk %s %s %s"), *SDKToString(SDK), *this->MetalLibraryBinary, Parameters);
 		return ExecGenericCommand(*XcrunPath, *BuiltParams, OutReturnCode, OutStdOut, OutStdErr);
 	}
 	else
