Fix Windows build.

#rnx

diff --git a/Engine/Source/Developer/Apple/MetalShaderFormat/Private/MetalShaderFormat.cpp b/Engine/Source/Developer/Apple/MetalShaderFormat/Private/MetalShaderFormat.cpp
--- a/Engine/Source/Developer/Apple/MetalShaderFormat/Private/MetalShaderFormat.cpp
+++ b/Engine/Source/Developer/Apple/MetalShaderFormat/Private/MetalShaderFormat.cpp
@@ -384,11 +384,12 @@ static FMetalCompilerToolchain::EMetalToolchainStatus ParseCompilerVersionAndTar
 
 		int32 Major = 0, Minor = 0, Patch = 0;
 		int32 NumResults = 0;
-		char AppleToolName[PATH_MAX] = { '\0' };
 #if !PLATFORM_WINDOWS
+		char AppleToolName[PATH_MAX] = { '\0' };
 		NumResults = sscanf(TCHAR_TO_ANSI(*Version), "Apple %s version %d.%d.%d", AppleToolName, &Major, &Minor, &Patch);
 #else
-		NumResults = swscanf_s(*Version, TEXT("Apple %s version %d.%d.%d"), AppleToolName, &Major, &Minor, &Patch);
+		TCHAR AppleToolName[WINDOWS_MAX_PATH] = TEXT("");
+		NumResults = swscanf_s(*Version, TEXT("Apple %ls version %d.%d.%d"), AppleToolName, WINDOWS_MAX_PATH, &Major, &Minor, &Patch);
 #endif
 		PackedVersionNumber.Major = Major;
 		PackedVersionNumber.Minor = Minor;
