Handle VersionEnum=7 in Metal backend and improve error output on unknown VersionEnum value.
Fix for new shader version introduced in CL 17643752.

#rnx

diff --git a/Engine/Source/Developer/Apple/MetalShaderFormat/Private/MetalShaderCompiler.cpp b/Engine/Source/Developer/Apple/MetalShaderFormat/Private/MetalShaderCompiler.cpp
--- a/Engine/Source/Developer/Apple/MetalShaderFormat/Private/MetalShaderCompiler.cpp
+++ b/Engine/Source/Developer/Apple/MetalShaderFormat/Private/MetalShaderCompiler.cpp
@@ -957,6 +957,7 @@ void CompileShader_Metal(const FShaderCompilerInput& _Input,FShaderCompilerOutpu
 	FString StandardVersion;
 	switch(VersionEnum)
     {
+		case 7:
 		case 6:
         case 5:
             // Enable full SM5 feature support so tessellation & fragment UAVs compile
@@ -1057,18 +1058,27 @@ void CompileShader_Metal(const FShaderCompilerInput& _Input,FShaderCompilerOutpu
 			return;
 		}
 		case 0:
-		default:
 		{
 			check(bIsMobile);
 			StandardVersion = TEXT("1.0");
 			MinOSVersion = TEXT("");
-			
+
+			Output.bSucceeded = false;
+			FShaderCompilerError* NewError = new(Output.Errors) FShaderCompilerError();
+			NewError->StrippedErrorMessage = FString::Printf(TEXT("Metal %s is no longer supported in UE."), *StandardVersion);
+			return;
+		}
+		default:
+		{
+			check(bIsMobile);
+			MinOSVersion = TEXT("");
+
 			Output.bSucceeded = false;
 			FShaderCompilerError* NewError = new(Output.Errors) FShaderCompilerError();
 			NewError->StrippedErrorMessage = FString::Printf(
-															 TEXT("Metal %s is no longer supported in UE."),
-															 *StandardVersion
-															 );
+				TEXT("Unknown Metal version: VersionEnum = %d; MAX_SHADER_LANGUAGE_VERSION = %s."),
+				*StandardVersion, VersionEnum, **MaxVersion
+			);
 			return;
 		}
 	}
