Replace EHlslCompileTarget with FHlslShaderModel in ShaderConductor interface for better control over input shader version.

#rnx

diff --git a/Engine/Source/Developer/Apple/MetalShaderFormat/Private/MetalDerivedData.cpp b/Engine/Source/Developer/Apple/MetalShaderFormat/Private/MetalDerivedData.cpp
--- a/Engine/Source/Developer/Apple/MetalShaderFormat/Private/MetalDerivedData.cpp
+++ b/Engine/Source/Developer/Apple/MetalShaderFormat/Private/MetalDerivedData.cpp
@@ -163,7 +163,7 @@ bool FMetalShaderBytecodeCooker::Build(TArray<uint8>& OutData)
 	return bSuccess;
 }
 
-FMetalShaderOutputCooker::FMetalShaderOutputCooker(const FShaderCompilerInput& _Input, FShaderCompilerOutput& _Output, const FString& _WorkingDirectory, FString _PreprocessedShader, FSHAHash _GUIDHash, uint8 _VersionEnum, uint32 _CCFlags, EHlslCompileTarget _HlslCompilerTarget, EHlslCompileTarget _MetalCompilerTarget, EMetalGPUSemantics _Semantics, EMetalTypeBufferMode _TypeMode, uint32 _MaxUnrollLoops, EShaderFrequency _Frequency, bool _bDumpDebugInfo, FString _Standard, FString _MinOSVersion)
+FMetalShaderOutputCooker::FMetalShaderOutputCooker(const FShaderCompilerInput& _Input, FShaderCompilerOutput& _Output, const FString& _WorkingDirectory, FString _PreprocessedShader, FSHAHash _GUIDHash, uint8 _VersionEnum, uint32 _CCFlags, EMetalGPUSemantics _Semantics, EMetalTypeBufferMode _TypeMode, uint32 _MaxUnrollLoops, EShaderFrequency _Frequency, bool _bDumpDebugInfo, FString _Standard, FString _MinOSVersion)
 	: Input(_Input)
 	, Output(_Output)
 	, WorkingDirectory(_WorkingDirectory)
@@ -172,8 +172,6 @@ FMetalShaderOutputCooker::FMetalShaderOutputCooker(const FShaderCompilerInput& _
 	, VersionEnum(_VersionEnum)
 	, CCFlags(_CCFlags)
 	, IABTier(0)
-	, HlslCompilerTarget(_HlslCompilerTarget)
-	, MetalCompilerTarget(_MetalCompilerTarget)
 	, Semantics(_Semantics)
 	, TypeMode(_TypeMode)
 	, MaxUnrollLoops(_MaxUnrollLoops)
diff --git a/Engine/Source/Developer/Apple/MetalShaderFormat/Private/MetalDerivedData.h b/Engine/Source/Developer/Apple/MetalShaderFormat/Private/MetalDerivedData.h
--- a/Engine/Source/Developer/Apple/MetalShaderFormat/Private/MetalDerivedData.h
+++ b/Engine/Source/Developer/Apple/MetalShaderFormat/Private/MetalDerivedData.h
@@ -135,7 +135,7 @@ struct FMetalShaderPreprocessed
 class FMetalShaderOutputCooker : public FDerivedDataPluginInterface
 {
 public:
-	FMetalShaderOutputCooker(const FShaderCompilerInput& _Input,FShaderCompilerOutput& Output,const FString& WorkingDirectory, FString PreprocessedShader, FSHAHash GUIDHash, uint8 VersionEnum, uint32 CCFlags, EHlslCompileTarget HlslCompilerTarget, EHlslCompileTarget MetalCompilerTarget, EMetalGPUSemantics Semantics, EMetalTypeBufferMode TypeMode, uint32 MaxUnrollLoops, EShaderFrequency Frequency, bool bDumpDebugInfo, FString Standard, FString MinOSVersion);
+	FMetalShaderOutputCooker(const FShaderCompilerInput& _Input,FShaderCompilerOutput& Output,const FString& WorkingDirectory, FString PreprocessedShader, FSHAHash GUIDHash, uint8 VersionEnum, uint32 CCFlags, EMetalGPUSemantics Semantics, EMetalTypeBufferMode TypeMode, uint32 MaxUnrollLoops, EShaderFrequency Frequency, bool bDumpDebugInfo, FString Standard, FString MinOSVersion);
 	virtual ~FMetalShaderOutputCooker();
 
 #if PLATFORM_MAC || PLATFORM_IOS
@@ -156,8 +156,6 @@ private:
 	uint8 VersionEnum;
 	uint32 CCFlags;
 	int32 IABTier;
-	EHlslCompileTarget HlslCompilerTarget;
-	EHlslCompileTarget MetalCompilerTarget;
 	EMetalGPUSemantics Semantics;
 	EMetalTypeBufferMode TypeMode;
 	uint32 MaxUnrollLoops;
diff --git a/Engine/Source/Developer/Apple/MetalShaderFormat/Private/MetalShaderCompiler.cpp b/Engine/Source/Developer/Apple/MetalShaderFormat/Private/MetalShaderCompiler.cpp
--- a/Engine/Source/Developer/Apple/MetalShaderFormat/Private/MetalShaderCompiler.cpp
+++ b/Engine/Source/Developer/Apple/MetalShaderFormat/Private/MetalShaderCompiler.cpp
@@ -864,8 +864,6 @@ void CompileShader_Metal(const FShaderCompilerInput& _Input,FShaderCompilerOutpu
 	auto Input = _Input;
 	FString PreprocessedShader;
 	FShaderCompilerDefinitions AdditionalDefines;
-	EHlslCompileTarget HlslCompilerTarget = HCT_FeatureLevelES3_1; // Always ES3.1 for now due to the way RCO has configured the MetalBackend
-	EHlslCompileTarget MetalCompilerTarget = HCT_FeatureLevelES3_1; // Varies depending on the actual intended Metal target.
 
 	// Work out which standard we need, this is dependent on the shader platform.
 	// TODO: Read from toolchain class
@@ -914,7 +912,6 @@ void CompileShader_Metal(const FShaderCompilerInput& _Input,FShaderCompilerOutpu
 		UE_CLOG(VersionEnum < 4, LogShaders, Warning, TEXT("Metal shader version must be Metal v2.1 or higher for format %s!"), VersionEnum, *Input.ShaderFormat.ToString());
 		AdditionalDefines.SetDefine(TEXT("METAL_MRT_PROFILE"), 1);
 		VersionEnum = VersionEnum >= 4 ? VersionEnum : 4;
-		MetalCompilerTarget = HCT_FeatureLevelSM5;
 		Semantics = EMetalGPUSemanticsTBDRDesktop;
 	}
 	else if (Input.ShaderFormat == NAME_SF_METAL_MACES3_1)
@@ -922,7 +919,6 @@ void CompileShader_Metal(const FShaderCompilerInput& _Input,FShaderCompilerOutpu
 		UE_CLOG(VersionEnum < 3, LogShaders, Warning, TEXT("Metal shader version must be Metal v2.0 or higher for format %s!"), VersionEnum, *Input.ShaderFormat.ToString());
 		AdditionalDefines.SetDefine(TEXT("METAL_PROFILE"), 1);
 		VersionEnum = VersionEnum > 3 ? VersionEnum : 3;
-		MetalCompilerTarget = HCT_FeatureLevelES3_1;
 		Semantics = EMetalGPUSemanticsImmediateDesktop;
 	}
 	else if (Input.ShaderFormat == NAME_SF_METAL_SM5)
@@ -931,7 +927,6 @@ void CompileShader_Metal(const FShaderCompilerInput& _Input,FShaderCompilerOutpu
 		AdditionalDefines.SetDefine(TEXT("METAL_SM5_PROFILE"), 1);
 		AdditionalDefines.SetDefine(TEXT("USING_VERTEX_SHADER_LAYER"), 1);
 		VersionEnum = VersionEnum > 3 ? VersionEnum : 3;
-		MetalCompilerTarget = HCT_FeatureLevelSM5;
 		Semantics = EMetalGPUSemanticsImmediateDesktop;
 	}
 	else if (Input.ShaderFormat == NAME_SF_METAL_MRT_MAC)
@@ -939,7 +934,6 @@ void CompileShader_Metal(const FShaderCompilerInput& _Input,FShaderCompilerOutpu
 		UE_CLOG(VersionEnum < 3, LogShaders, Warning, TEXT("Metal shader version must be Metal v2.0 or higher for format %s!"), VersionEnum, *Input.ShaderFormat.ToString());
 		AdditionalDefines.SetDefine(TEXT("METAL_MRT_PROFILE"), 1);
 		VersionEnum = VersionEnum > 3 ? VersionEnum : 3;
-		MetalCompilerTarget = HCT_FeatureLevelSM5;
 		Semantics = EMetalGPUSemanticsTBDRDesktop;
 	}
 	else
@@ -1214,7 +1208,7 @@ void CompileShader_Metal(const FShaderCompilerInput& _Input,FShaderCompilerOutpu
 		FSHA1::HashBuffer(&Guid, sizeof(FGuid), GUIDHash.Hash);
 	}
 
-	FMetalShaderOutputCooker* Cooker = new FMetalShaderOutputCooker(Input,Output,WorkingDirectory, PreprocessedShader, GUIDHash, VersionEnum, CCFlags, HlslCompilerTarget, MetalCompilerTarget, Semantics, TypeMode, MaxUnrollLoops, Frequency, bDumpDebugInfo, Standard, MinOSVersion);
+	FMetalShaderOutputCooker* Cooker = new FMetalShaderOutputCooker(Input,Output,WorkingDirectory, PreprocessedShader, GUIDHash, VersionEnum, CCFlags, Semantics, TypeMode, MaxUnrollLoops, Frequency, bDumpDebugInfo, Standard, MinOSVersion);
 		
 	bool bDataWasBuilt = false;
 	TArray<uint8> OutData;
diff --git a/Engine/Source/Developer/ShaderCompilerCommon/Private/ShaderConductorContext.cpp b/Engine/Source/Developer/ShaderCompilerCommon/Private/ShaderConductorContext.cpp
--- a/Engine/Source/Developer/ShaderCompilerCommon/Private/ShaderConductorContext.cpp
+++ b/Engine/Source/Developer/ShaderCompilerCommon/Private/ShaderConductorContext.cpp
@@ -193,19 +193,6 @@ namespace CrossCompiler
 		return ShaderConductor::ShaderStage::NumShaderStages;
 	}
 
-	static ShaderConductor::Compiler::ShaderModel ToShaderConductorShaderModel(EHlslCompileTarget Target)
-	{
-		switch (Target)
-		{
-		case HCT_FeatureLevelSM4:		return { 4, 0 };
-		case HCT_FeatureLevelES3_1Ext:	return { 4, 0 };
-		case HCT_FeatureLevelSM5:		return { 5, 0 };
-		case HCT_FeatureLevelES3_1:		return { 4, 0 };
-		default: checkf(0, TEXT("Invalid input shader target for enum <EHlslCompileTarget>."));
-		}
-		return { 6,0 };
-	}
-
 	// Wrapper structure to hold all intermediate buffers for ShaderConductor
 	struct FShaderConductorContext::FShaderConductorIntermediates
 	{
@@ -397,22 +384,24 @@ namespace CrossCompiler
 
 	static void ConvertScOptions(const FShaderConductorOptions& InOptions, ShaderConductor::Compiler::Options& OutOptions, const TArray<ANSICHAR const*>& ExtraDxcArgRefs)
 	{
+		// Validate input shader model with respect to certain language features.
+		checkf(
+			(!InOptions.bEnable16bitTypes || InOptions.ShaderModel >= FHlslShaderModel{ 6, 2 }),
+			TEXT("DXC option '-enable-16bit-types' only supported with SM6.2+ but SM%u.%u was specified"),
+			InOptions.ShaderModel.Major, InOptions.ShaderModel.Minor
+		);
+
 		OutOptions.removeUnusedGlobals = InOptions.bRemoveUnusedGlobals;
 		OutOptions.packMatricesInRowMajor = InOptions.bPackMatricesInRowMajor;
 		OutOptions.enable16bitTypes = InOptions.bEnable16bitTypes;
 		OutOptions.enableDebugInfo = InOptions.bEnableDebugInfo;
 		OutOptions.disableOptimizations = InOptions.bDisableOptimizations;
 		OutOptions.enableFMAPass = InOptions.bEnableFMAPass;
-
-		if (OutOptions.enable16bitTypes)
-		{
-			// 16-bit types only supports with SM 6.2+
-			OutOptions.shaderModel = ShaderConductor::Compiler::ShaderModel{ 6, 2 };
-		}
-		else
+		OutOptions.shaderModel = ShaderConductor::Compiler::ShaderModel
 		{
-			OutOptions.shaderModel = ToShaderConductorShaderModel(InOptions.TargetProfile);
-		}
+			static_cast<uint8>(InOptions.ShaderModel.Major),
+			static_cast<uint8>(InOptions.ShaderModel.Minor)
+		};
 
 		if (ExtraDxcArgRefs.Num() > 0)
 		{
diff --git a/Engine/Source/Developer/ShaderCompilerCommon/Public/ShaderConductorContext.h b/Engine/Source/Developer/ShaderCompilerCommon/Public/ShaderConductorContext.h
--- a/Engine/Source/Developer/ShaderCompilerCommon/Public/ShaderConductorContext.h
+++ b/Engine/Source/Developer/ShaderCompilerCommon/Public/ShaderConductorContext.h
@@ -11,6 +11,46 @@
 namespace CrossCompiler
 {
 
+	/** Shader model version for HLSL input language. */
+	struct SHADERCOMPILERCOMMON_API FHlslShaderModel
+	{
+		/** Major shader model version (e.g. 6 in SM6.2). */
+		uint16 Major;
+
+		/** Minor shader model version (e.g. 2 in SM6.2). */
+		uint16 Minor;
+
+		FORCEINLINE bool operator == (const FHlslShaderModel& Rhs) const
+		{
+			return Major == Rhs.Major && Minor == Rhs.Minor;
+		}
+
+		FORCEINLINE bool operator != (const FHlslShaderModel& Rhs) const
+		{
+			return !(*this == Rhs);
+		}
+
+		FORCEINLINE bool operator < (const FHlslShaderModel& Rhs) const
+		{
+			return Major < Rhs.Major || (Major == Rhs.Major && Minor < Rhs.Minor);
+		}
+
+		FORCEINLINE bool operator <= (const FHlslShaderModel& Rhs) const
+		{
+			return *this < Rhs || *this == Rhs;
+		}
+
+		FORCEINLINE bool operator > (const FHlslShaderModel& Rhs) const
+		{
+			return Rhs < *this;
+		}
+
+		FORCEINLINE bool operator >= (const FHlslShaderModel& Rhs) const
+		{
+			return *this > Rhs || *this == Rhs;
+		}
+	};
+
 	/** Wrapper structure to pass options descriptor to ShaderConductor. This is mapped to <struct ShaderConductor::Compiler::Options>. */
 	struct SHADERCOMPILERCOMMON_API FShaderConductorOptions
 	{
@@ -32,8 +72,8 @@ namespace CrossCompiler
 		/** Enable a pass that converts floating point MUL+ADD pairs into FMAs to avoid re-association. */
 		bool bEnableFMAPass = false;
 
-		/** Target shader profile. By default HCT_FeatureLevelSM5. */
-		EHlslCompileTarget TargetProfile = HCT_FeatureLevelSM5;
+		/** Shader model version of the input language. By default SM6.2. */
+		FHlslShaderModel ShaderModel = { 6, 2 };
 	};
 
 	/** Target high level languages for ShaderConductor output. */
diff --git a/Engine/Source/Developer/VulkanShaderFormat/Private/VulkanShaderCompiler.cpp b/Engine/Source/Developer/VulkanShaderFormat/Private/VulkanShaderCompiler.cpp
--- a/Engine/Source/Developer/VulkanShaderFormat/Private/VulkanShaderCompiler.cpp
+++ b/Engine/Source/Developer/VulkanShaderFormat/Private/VulkanShaderCompiler.cpp
@@ -1949,7 +1949,6 @@ static bool CompileWithShaderConductor(
 	const FString&			EntryPointName,
 	EShaderFrequency		Frequency,
 	const FCompilerInfo&	CompilerInfo,
-	EHlslCompileTarget		HlslCompilerTarget,
 	FShaderCompilerOutput&	Output,
 	FVulkanBindingTable&	BindingTable,
 	bool					bStripReflect)
@@ -1976,7 +1975,6 @@ static bool CompileWithShaderConductor(
 
 	// Initialize compilation options for ShaderConductor
 	CrossCompiler::FShaderConductorOptions Options;
-	Options.TargetProfile = HlslCompilerTarget;
 
 	if (bRewriteHlslSource)
 	{
@@ -2064,21 +2062,15 @@ void DoCompileVulkanShader(const FShaderCompilerInput& Input, FShaderCompilerOut
 
 	FString PreprocessedShader;
 	FShaderCompilerDefinitions AdditionalDefines;
-	EHlslCompileTarget HlslCompilerTarget = HCT_FeatureLevelES3_1Ext;
-	EHlslCompileTarget HlslCompilerTargetES = HCT_FeatureLevelES3_1Ext;
 	AdditionalDefines.SetDefine(TEXT("COMPILER_HLSLCC"), 1);
 	AdditionalDefines.SetDefine(TEXT("COMPILER_VULKAN"), 1);
 	if(bIsMobile)
 	{
-		HlslCompilerTarget = HCT_FeatureLevelES3_1Ext;
-		HlslCompilerTargetES = HCT_FeatureLevelES3_1Ext;
 		AdditionalDefines.SetDefine(TEXT("ES3_1_PROFILE"), 1);
 		AdditionalDefines.SetDefine(TEXT("VULKAN_PROFILE"), 1);
 	}
 	else if (bIsSM5)
 	{
-		HlslCompilerTarget = HCT_FeatureLevelSM5;
-		HlslCompilerTargetES = HCT_FeatureLevelSM5;
 		AdditionalDefines.SetDefine(TEXT("VULKAN_PROFILE_SM5"), 1);
 	}
 	AdditionalDefines.SetDefine(TEXT("row_major"), TEXT(""));
@@ -2177,7 +2169,7 @@ void DoCompileVulkanShader(const FShaderCompilerInput& Input, FShaderCompilerOut
 
 #if PLATFORM_MAC || PLATFORM_WINDOWS || PLATFORM_LINUX
 	// Cross-compile shader via ShaderConductor (DXC, SPIRV-Tools, SPIRV-Cross)
-	bSuccess = CompileWithShaderConductor(PreprocessedShaderSource, EntryPointName, Frequency, CompilerInfo, HlslCompilerTarget, Output, BindingTable, bStripReflect);
+	bSuccess = CompileWithShaderConductor(PreprocessedShaderSource, EntryPointName, Frequency, CompilerInfo, Output, BindingTable, bStripReflect);
 #endif // PLATFORM_MAC || PLATFORM_WINDOWS || PLATFORM_LINUX
 	
 	ShaderParameterParser.ValidateShaderParameterTypes(Input, Output);
