Fix compile error in MetalShaderCompiler.cpp as appendix to CL 17817375.

#rnx

diff --git a/Engine/Source/Developer/Apple/MetalShaderFormat/Private/MetalShaderCompiler.cpp b/Engine/Source/Developer/Apple/MetalShaderFormat/Private/MetalShaderCompiler.cpp
--- a/Engine/Source/Developer/Apple/MetalShaderFormat/Private/MetalShaderCompiler.cpp
+++ b/Engine/Source/Developer/Apple/MetalShaderFormat/Private/MetalShaderCompiler.cpp
@@ -956,7 +956,6 @@ void CompileShader_Metal(const FShaderCompilerInput& _Input,FShaderCompilerOutpu
         case 5:
             // Enable full SM5 feature support so tessellation & fragment UAVs compile
             TypeMode = EMetalTypeBufferModeTB;
-            HlslCompilerTarget = HCT_FeatureLevelSM5;
             StandardVersion = TEXT("2.1");
             if (bAppleTV)
             {
@@ -974,7 +973,6 @@ void CompileShader_Metal(const FShaderCompilerInput& _Input,FShaderCompilerOutpu
 		case 4:
 			// Enable full SM5 feature support so tessellation & fragment UAVs compile
 			TypeMode = EMetalTypeBufferModeTB;
-            HlslCompilerTarget = HCT_FeatureLevelSM5;
 			StandardVersion = TEXT("2.1");
 			if (bAppleTV)
 			{
@@ -994,7 +992,6 @@ void CompileShader_Metal(const FShaderCompilerInput& _Input,FShaderCompilerOutpu
 		case 3:
 			// Enable full SM5 feature support so tessellation & fragment UAVs compile
 			TypeMode = EMetalTypeBufferMode2D;
-            HlslCompilerTarget = HCT_FeatureLevelSM5;
 			StandardVersion = TEXT("2.0");
 			if (bAppleTV)
 			{
@@ -1014,7 +1011,6 @@ void CompileShader_Metal(const FShaderCompilerInput& _Input,FShaderCompilerOutpu
 		case 2:
 			// Enable full SM5 feature support so tessellation & fragment UAVs compile
 			TypeMode = EMetalTypeBufferMode2D;
-            HlslCompilerTarget = HCT_FeatureLevelSM5;
 			StandardVersion = TEXT("1.2");
 			if (bAppleTV)
 			{
@@ -1039,7 +1035,6 @@ void CompileShader_Metal(const FShaderCompilerInput& _Input,FShaderCompilerOutpu
 			break;
 		case 1:
 		{
-			HlslCompilerTarget = bIsMobile ? HlslCompilerTarget : HCT_FeatureLevelSM5;
 			StandardVersion = TEXT("1.1");
 			MinOSVersion = bIsMobile ? TEXT("") : TEXT("-mmacosx-version-min=10.11");
 			
