Enable Metal v2.4; maps from MAX_SHADER_LANGUAGE_VERSION=7.

#rnx

diff --git a/Engine/Source/Developer/Apple/MetalShaderFormat/Private/MetalDerivedData.cpp b/Engine/Source/Developer/Apple/MetalShaderFormat/Private/MetalDerivedData.cpp
--- a/Engine/Source/Developer/Apple/MetalShaderFormat/Private/MetalDerivedData.cpp
+++ b/Engine/Source/Developer/Apple/MetalShaderFormat/Private/MetalDerivedData.cpp
@@ -1110,6 +1110,11 @@ bool FMetalShaderOutputCooker::Build(TArray<uint8>& OutData)
 			
 			switch (VersionEnum)
 			{
+				case 7:
+				{
+					TargetDesc.Version = 20400;
+					break;
+				}
 				case 6:
 				{
 					TargetDesc.Version = 20300;
diff --git a/Engine/Source/Developer/ShaderCompilerCommon/Private/ShaderConductorContext.cpp b/Engine/Source/Developer/ShaderCompilerCommon/Private/ShaderConductorContext.cpp
--- a/Engine/Source/Developer/ShaderCompilerCommon/Private/ShaderConductorContext.cpp
+++ b/Engine/Source/Developer/ShaderCompilerCommon/Private/ShaderConductorContext.cpp
@@ -283,6 +283,7 @@ namespace CrossCompiler
 	{
 		switch (Version)
 		{
+		case 20400: return "20400";
 		case 20300: return "20300";
 		case 20200: return "20200";
 		case 20100: return "20100";
