Bring missing changes from Release-Engine-Staging to Release-Engine-Test

diff --git a/Engine/Plugins/Bridge/Bridge.uplugin b/Engine/Plugins/Bridge/Bridge.uplugin
--- a/Engine/Plugins/Bridge/Bridge.uplugin
+++ b/Engine/Plugins/Bridge/Bridge.uplugin
@@ -10,6 +10,7 @@
 	"DocsURL": "https://help.quixel.com/hc/en-us/sections/360005846137-Quixel-Bridge-for-Unreal-Engine-5",
 	"MarketplaceURL": "",
 	"SupportURL": "[REDACTED_EMAIL]",
+	"EngineVersion": "5.0.0",
 	"CanContainContent": true,
 	"EnabledByDefault": true,
 	"SupportedTargetPlatforms": [
diff --git a/Engine/Source/Developer/Apple/MetalShaderFormat/Private/MetalShaderFormat.cpp b/Engine/Source/Developer/Apple/MetalShaderFormat/Private/MetalShaderFormat.cpp
--- a/Engine/Source/Developer/Apple/MetalShaderFormat/Private/MetalShaderFormat.cpp
+++ b/Engine/Source/Developer/Apple/MetalShaderFormat/Private/MetalShaderFormat.cpp
@@ -389,7 +389,7 @@ static FMetalCompilerToolchain::EMetalToolchainStatus ParseCompilerVersionAndTar
 		char AppleToolName[PATH_MAX] = { '\0' };
 		NumResults = sscanf(TCHAR_TO_ANSI(*Version), "Apple %s version %d.%d.%d", AppleToolName, &Major, &Minor, &Patch);
 #else
-		TCHAR AppleToolName[WINDOWS_MAX_PATH] = TEXT("");
+		TCHAR AppleToolName[WINDOWS_MAX_PATH] = { '\0' };
 		NumResults = swscanf_s(*Version, TEXT("Apple %ls version %d.%d.%d"), AppleToolName, WINDOWS_MAX_PATH, &Major, &Minor, &Patch);
 #endif
 		PackedVersionNumber.Major = Major;
