Improve stability for out of bounds buffer slots in Metal backend.

#rnx

diff --git a/Engine/Source/Developer/Apple/MetalShaderFormat/Private/MetalDerivedData.cpp b/Engine/Source/Developer/Apple/MetalShaderFormat/Private/MetalDerivedData.cpp
--- a/Engine/Source/Developer/Apple/MetalShaderFormat/Private/MetalDerivedData.cpp
+++ b/Engine/Source/Developer/Apple/MetalShaderFormat/Private/MetalDerivedData.cpp
@@ -1244,6 +1244,7 @@ bool DoCompileMetalShader(
 	{
 		Output.Target = Input.Target;
 		BuildMetalShaderOutput(Output, Input, GUIDHash, CCFlags, MetalSource.c_str(), MetalSource.length(), CRCLen, CRC, VersionEnum, *Standard, *MinOSVersion, TypeMode, Output.Errors, OutputData.TypedBuffers, OutputData.InvariantBuffers, OutputData.TypedUAVs, OutputData.ConstantBuffers, OutputData.TypedBufferFormats, bAllowFastIntriniscs);
+		return Output.bSucceeded;
 	}
 	else
 	{
diff --git a/Engine/Source/Developer/Apple/MetalShaderFormat/Private/MetalShaderCompiler.cpp b/Engine/Source/Developer/Apple/MetalShaderFormat/Private/MetalShaderCompiler.cpp
--- a/Engine/Source/Developer/Apple/MetalShaderFormat/Private/MetalShaderCompiler.cpp
+++ b/Engine/Source/Developer/Apple/MetalShaderFormat/Private/MetalShaderCompiler.cpp
@@ -49,6 +49,9 @@ static bool CompileProcessAllowsRuntimeShaderCompiling(const FShaderCompilerInpu
 	return !bArchiving && bDebug;
 }
 
+constexpr uint16 GMetalMaxUniformBufferSlots = 32;
+constexpr int32 GMetalDefaultShadingLanguageVersion = 7;
+
 /*------------------------------------------------------------------------------
 	Shader compiling.
 ------------------------------------------------------------------------------*/
@@ -357,6 +360,7 @@ void BuildMetalShaderOutput(
 	}
 
 	// Then 'normal' uniform buffers.
+	bool bOutOfBounds = false;
 	for (auto& UniformBlock : CCHeader.UniformBlocks)
 	{
 		uint16 UBIndex = UniformBlock.Index;
@@ -364,9 +368,21 @@ void BuildMetalShaderOutput(
 		{
 			Header.Bindings.NumUniformBuffers = UBIndex + 1;
 		}
+		if (UBIndex >= GMetalMaxUniformBufferSlots)
+		{
+			bOutOfBounds = true;
+			new(OutErrors) FShaderCompilerError(*FString::Printf(TEXT("Uniform buffer index (%d) exceeded upper bound of slots (%d) for Metal API: %s"), UBIndex, GMetalMaxUniformBufferSlots, *UniformBlock.Name));
+			continue;
+		}
 		UsedUniformBufferSlots[UBIndex] = true;
 		ParameterMap.AddParameterAllocation(*UniformBlock.Name, UBIndex, 0, 0, EShaderParameterType::UniformBuffer);
 	}
+	
+	if (bOutOfBounds)
+	{
+		ShaderOutput.bSucceeded = false;
+		return;
+	}
 
 	// Packed global uniforms
 	const uint16 BytesPerComponent = 4;
@@ -545,7 +561,7 @@ void BuildMetalShaderOutput(
 	if (Header.Bindings.NumSamplers > MaxMetalSamplers)
 	{
 		ShaderOutput.bSucceeded = false;
-		FShaderCompilerError* NewError = new(ShaderOutput.Errors) FShaderCompilerError();
+		FShaderCompilerError* NewError = new(OutErrors) FShaderCompilerError();
 		
 		FString SamplerList;
 		for (int32 i = 0; i < CCHeader.SamplerStates.Num(); i++)
@@ -845,7 +861,18 @@ void CompileShader_Metal(const FShaderCompilerInput& _Input,FShaderCompilerOutpu
 	
 	EMetalGPUSemantics Semantics = EMetalGPUSemanticsMobile;
 	
-    int32 VersionEnum = FCString::Atoi(*FString(*Input.Environment.GetDefinitions().Find(TEXT("SHADER_LANGUAGE_VERSION"))));
+    const int32 VersionEnum = [&Input, &Output]() -> int32
+	{
+		if (const FString* VersionEnumEntry = Input.Environment.GetDefinitions().Find(TEXT("SHADER_LANGUAGE_VERSION")))
+		{
+			return FCString::Atoi(*FString(*VersionEnumEntry));
+		}
+		else
+		{
+			new(Output.Errors) FShaderCompilerError(*FString::Printf(TEXT("Missing definition of SHADER_LANGUAGE_VERSION; Falling back to default value %d"), GMetalDefaultShadingLanguageVersion));
+			return GMetalDefaultShadingLanguageVersion;
+		}
+	}();
 
 	// TODO read from toolchain
 	bool bAppleTV = (Input.ShaderFormat == NAME_SF_METAL_TVOS || Input.ShaderFormat == NAME_SF_METAL_MRT_TVOS);
