Fix for Metal running out of buffer binding slots (max 31) due to the allocation of texture+buffer slot per SRV texture.

- Consume Uniform then UAV buffer slots first.
- Don't use up a buffer slot for texel buffers as we should be using texture buffers for these now.
- If we are not using a buffer slot for the same texel slot index then don't attempt to bind the underlying SRV buffer when binding the SRV texture otherwise it'll stomp on other in use buffer slots.
- Use CountTrailingZeros64() when referencing TextureIndices as that is a 64bit type.

#rnx

diff --git a/Engine/Shaders/Public/Platform/Metal/MetalCommon.ush b/Engine/Shaders/Public/Platform/Metal/MetalCommon.ush
--- a/Engine/Shaders/Public/Platform/Metal/MetalCommon.ush
+++ b/Engine/Shaders/Public/Platform/Metal/MetalCommon.ush
@@ -5,7 +5,7 @@
 =============================================================================*/
 
 // Update this GUID to improve shader recompilation for Metal only shaders
-// GUID = 80604507-D068-4171-8020-B70A6982E681
+// GUID = 9e7b1ea2-6caf-4e79-bdec-fe12e8cc9430
 
 #pragma once
 
diff --git a/Engine/Source/Developer/Apple/MetalShaderFormat/Private/MetalDerivedData.cpp b/Engine/Source/Developer/Apple/MetalShaderFormat/Private/MetalDerivedData.cpp
--- a/Engine/Source/Developer/Apple/MetalShaderFormat/Private/MetalDerivedData.cpp
+++ b/Engine/Source/Developer/Apple/MetalShaderFormat/Private/MetalDerivedData.cpp
@@ -635,23 +635,6 @@ bool DoCompileMetalShader(
 					CCHeaderWriter.WriteUniformBlock(*Name, SetIndex);
 				}
 				
-				for (auto const& Binding : ReflectionBindings.TBufferSRVs)
-				{
-					check(TextureIndices);
-					uint32 Index = FPlatformMath::CountTrailingZeros(TextureIndices);
-					
-					// No support for 3-component types in dxc/SPIRV/MSL - need to expose my workarounds there too
-					BufferIndices &= ~(1ull << (uint64)Index);
-					TextureIndices &= ~(1ull << uint64(Index));
-					
-					OutputData.TypedBuffers |= (1 << Index);
-					
-					CCHeaderWriter.WriteSRV(UTF8_TO_TCHAR(Binding->name), Index);
-					
-					SPVRResult = Reflection.ChangeDescriptorBindingNumbers(Binding, Index, GlobalSetId);
-					check(SPVRResult == SPV_REFLECT_RESULT_SUCCESS);
-				}
-				
 				for (auto const& Binding : ReflectionBindings.SBufferSRVs)
 				{
 					check(BufferIndices);
@@ -718,6 +701,20 @@ bool DoCompileMetalShader(
 					check(SPVRResult == SPV_REFLECT_RESULT_SUCCESS);
 				}
 				
+				for (auto const& Binding : ReflectionBindings.TBufferSRVs)
+				{
+					check(TextureIndices);
+					uint32 Index = FPlatformMath::CountTrailingZeros64(TextureIndices);
+					TextureIndices &= ~(1ull << uint64(Index));
+					
+					OutputData.TypedBuffers |= (1 << Index);
+					
+					CCHeaderWriter.WriteSRV(UTF8_TO_TCHAR(Binding->name), Index);
+					
+					SPVRResult = Reflection.ChangeDescriptorBindingNumbers(Binding, Index, GlobalSetId);
+					check(SPVRResult == SPV_REFLECT_RESULT_SUCCESS);
+				}
+				
 				for (auto const& Binding : ReflectionBindings.TextureSRVs)
 				{
 					check(TextureIndices);
diff --git a/Engine/Source/Runtime/Apple/MetalRHI/Private/MetalStateCache.cpp b/Engine/Source/Runtime/Apple/MetalRHI/Private/MetalStateCache.cpp
--- a/Engine/Source/Runtime/Apple/MetalRHI/Private/MetalStateCache.cpp
+++ b/Engine/Source/Runtime/Apple/MetalRHI/Private/MetalStateCache.cpp
@@ -1477,14 +1477,8 @@ void FMetalStateCache::SetShaderResourceView(FMetalContext* Context, EMetalShade
 			Tex = SRV->GetLinearTexture(false);
 			
 			SetShaderTexture(ShaderStage, Tex, BindIndex, mtlpp::ResourceUsage(mtlpp::ResourceUsage::Read|mtlpp::ResourceUsage::Sample));
-			if (VB)
-            {
-				SetShaderBuffer(ShaderStage, VB->GetCurrentBufferOrNil(), VB->Data, SRV->Offset, VB->GetSize(), BindIndex, mtlpp::ResourceUsage::Read, (EPixelFormat)SRV->Format);
-            }
-            else if (IB)
-            {
-                SetShaderBuffer(ShaderStage, IB->GetCurrentBufferOrNil(), nil, SRV->Offset, IB->GetSize(), BindIndex, mtlpp::ResourceUsage::Read, (EPixelFormat)SRV->Format);
-            }
+			
+			// Removed the binding of the underlying buffer for SRVs as we are now using texture buffers here
 		}
 		else if (VB)
 		{
