Reorganizing reflection handling to funnel through common functions, Metal edition.

diff --git a/Engine/Source/Developer/Apple/MetalShaderFormat/Private/MetalShaderCompiler.cpp b/Engine/Source/Developer/Apple/MetalShaderFormat/Private/MetalShaderCompiler.cpp
--- a/Engine/Source/Developer/Apple/MetalShaderFormat/Private/MetalShaderCompiler.cpp
+++ b/Engine/Source/Developer/Apple/MetalShaderFormat/Private/MetalShaderCompiler.cpp
@@ -375,7 +375,8 @@ void BuildMetalShaderOutput(
 			continue;
 		}
 		UsedUniformBufferSlots[UBIndex] = true;
-		ParameterMap.AddParameterAllocation(*UniformBlock.Name, UBIndex, 0, 0, EShaderParameterType::UniformBuffer);
+
+		HandleReflectedUniformBuffer(UniformBlock.Name, UBIndex, ShaderOutput);
 	}
 	
 	if (bOutOfBounds)
@@ -389,12 +390,12 @@ void BuildMetalShaderOutput(
 	TMap<ANSICHAR, uint16> PackedGlobalArraySize;
 	for (auto& PackedGlobal : CCHeader.PackedGlobals)
 	{
-		ParameterMap.AddParameterAllocation(
-			*PackedGlobal.Name,
+		HandleReflectedGlobalConstantBufferMember(
+			PackedGlobal.Name,
 			PackedGlobal.PackedType,
 			PackedGlobal.Offset * BytesPerComponent,
 			PackedGlobal.Count * BytesPerComponent,
-			EShaderParameterType::LooseData
+			ShaderOutput
 		);
 
 		uint16& Size = PackedGlobalArraySize.FindOrAdd(PackedGlobal.PackedType);
@@ -407,12 +408,12 @@ void BuildMetalShaderOutput(
 	{
 		for (auto& Member : PackedUB.Members)
 		{
-			ParameterMap.AddParameterAllocation(
-				*Member.Name,
-				(ANSICHAR)CrossCompiler::EPackedTypeName::HighP,
+			HandleReflectedGlobalConstantBufferMember(
+				Member.Name,
+				(uint32)CrossCompiler::EPackedTypeName::HighP,
 				Member.Offset * BytesPerComponent,
-												Member.Count * BytesPerComponent, 
-				EShaderParameterType::LooseData
+				Member.Count * BytesPerComponent,
+				ShaderOutput
 			);
 			
 			uint16& Size = PackedUniformBuffersSize.FindOrAdd(PackedUB.Attribute.Index).FindOrAdd(CrossCompiler::EPackedTypeName::HighP);
@@ -462,13 +463,7 @@ void BuildMetalShaderOutput(
 	TMap<FString, uint32> SamplerMap;
 	for (auto& Sampler : CCHeader.Samplers)
 	{
-		ParameterMap.AddParameterAllocation(
-			*Sampler.Name,
-			0,
-			Sampler.Offset,
-			Sampler.Count,
-			EShaderParameterType::SRV
-		);
+		HandleReflectedShaderResource(Sampler.Name, Sampler.Offset, Sampler.Count, ShaderOutput);
 
 		NumTextures += Sampler.Count;
 
@@ -483,13 +478,7 @@ void BuildMetalShaderOutput(
 	// Then UAVs (images in Metal)
 	for (auto& UAV : CCHeader.UAVs)
 	{
-		ParameterMap.AddParameterAllocation(
-			*UAV.Name,
-			0,
-			UAV.Offset,
-			UAV.Count,
-			EShaderParameterType::UAV
-		);
+		HandleReflectedShaderUAV(UAV.Name, UAV.Offset, UAV.Count, ShaderOutput);
 
 		Header.Bindings.NumUAVs = FMath::Max<uint8>(
 			Header.Bindings.NumSamplers,
@@ -504,13 +493,7 @@ void BuildMetalShaderOutput(
 			SamplerMap.Add(SamplerState.Name, 1);
 		}
 		
-		ParameterMap.AddParameterAllocation(
-			*SamplerState.Name,
-			0,
-			SamplerState.Index,
-			SamplerMap[SamplerState.Name],
-			EShaderParameterType::Sampler
-		);
+		HandleReflectedShaderResource(SamplerState.Name, SamplerState.Index, SamplerMap[SamplerState.Name], ShaderOutput);
 	}
 
 	Header.NumThreadsX = CCHeader.NumThreads[0];
