Work around macOS Monterey 12.3.1 issue spewing objc warning messages in Apple frameworks to stdout.

#rnx

diff --git a/Engine/Source/Developer/Apple/MetalShaderFormat/Private/MetalShaderFormat.cpp b/Engine/Source/Developer/Apple/MetalShaderFormat/Private/MetalShaderFormat.cpp
--- a/Engine/Source/Developer/Apple/MetalShaderFormat/Private/MetalShaderFormat.cpp
+++ b/Engine/Source/Developer/Apple/MetalShaderFormat/Private/MetalShaderFormat.cpp
@@ -380,9 +380,19 @@ static FMetalCompilerToolchain::EMetalToolchainStatus ParseCompilerVersionAndTar
 
 	TArray<FString> Lines;
 	OutputOfMetalDashV.ParseIntoArrayLines(Lines);
+	int32 VersionLineIndex = 0;
 	{
-		VersionString = Lines[0];
-		FString& Version = Lines[0];
+		for (int32 Index = 0; Index < Lines.Num(); ++Index)
+		{
+			if (Lines[Index].StartsWith(TEXT("Apple ")) && Lines[Index].Contains(TEXT(" version ")) && Lines[Index].EndsWith(TEXT(")")))
+			{
+				VersionLineIndex = Index;
+				break;
+			}
+		}
+
+		VersionString = Lines[VersionLineIndex];
+		FString& Version = Lines[VersionLineIndex];
 		check(!Version.IsEmpty());
 
 		int32 Major = 0, Minor = 0, Patch = 0;
@@ -405,7 +415,7 @@ static FMetalCompilerToolchain::EMetalToolchainStatus ParseCompilerVersionAndTar
 	}
 
 	{
-		FString& FormatVersion = Lines[1];
+		FString& FormatVersion = Lines[VersionLineIndex + 1];
 		int32 Major = 0, Minor = 0, Patch = 0;
 		int32 NumResults = 0;
 #if !PLATFORM_WINDOWS
@@ -433,7 +443,18 @@ static FMetalCompilerToolchain::EMetalToolchainStatus ParseLibraryToolpath(const
 	TArray<FString> Lines;
 	OutputOfMetalSearchDirs.ParseIntoArrayLines(Lines);
 	{
-		FString& LibraryLine = Lines[1];
+		int32 LibrariesLineIndex = 0;
+		
+		for (int32 Index = 0; Index < Lines.Num(); ++Index)
+		{
+			if (Lines[Index].StartsWith(TEXT("libraries: =")))
+			{
+				LibrariesLineIndex = Index;
+				break;
+			}
+		}
+
+		FString& LibraryLine = Lines[LibrariesLineIndex];
 
 		LibraryPath = LibraryLine.RightChop(LibraryPrefix.Len());
 
@@ -730,7 +751,7 @@ FMetalCompilerToolchain::EMetalToolchainStatus FMetalCompilerToolchain::DoMacNat
 
 	int32 ReturnCode = 0;
 	FString StdOut, StdErr;
-	bool bSuccess = this->ExecGenericCommand(*XcrunPath, *FString::Printf(TEXT("-sdk %s --find %s"), *this->MetalMacSDK, *this->MetalFrontendBinary), &ReturnCode, &StdOut, &StdErr);
+	bool bSuccess = this->ExecGenericCommand(*XcrunPath, *FString::Printf(TEXT("--sdk %s --find %s"), *this->MetalMacSDK, *this->MetalFrontendBinary), &ReturnCode, &StdOut, &StdErr);
 	bSuccess |= FPaths::FileExists(StdOut);
 	
 	if(!bSuccess || ReturnCode > 0)
