Fix for metal shader backend mismatch between resources and samplers

diff --git a/Engine/Shaders/Public/ShaderVersion.ush b/Engine/Shaders/Public/ShaderVersion.ush
--- a/Engine/Shaders/Public/ShaderVersion.ush
+++ b/Engine/Shaders/Public/ShaderVersion.ush
@@ -2,4 +2,4 @@
 // Each time the console command is executed it generates a new GUID. As this file is included
 // in Platform.ush (which should be included in any shader) it allows to invalidate the shader DDC.
 // 
-// GUID = 15749AA0-E6E1-44D9-8D85-687EE3C73A20
+// GUID = 0E1525E8-83A2-4FB1-BB2E-8BABCE046416
diff --git a/Engine/Source/Developer/Apple/MetalShaderFormat/Private/MetalShaderCompiler.cpp b/Engine/Source/Developer/Apple/MetalShaderFormat/Private/MetalShaderCompiler.cpp
--- a/Engine/Source/Developer/Apple/MetalShaderFormat/Private/MetalShaderCompiler.cpp
+++ b/Engine/Source/Developer/Apple/MetalShaderFormat/Private/MetalShaderCompiler.cpp
@@ -493,7 +493,7 @@ void BuildMetalShaderOutput(
 			SamplerMap.Add(SamplerState.Name, 1);
 		}
 		
-		HandleReflectedShaderResource(SamplerState.Name, SamplerState.Index, SamplerMap[SamplerState.Name], ShaderOutput);
+		HandleReflectedShaderSampler(SamplerState.Name, SamplerState.Index, SamplerMap[SamplerState.Name], ShaderOutput);
 	}
 
 	Header.NumThreadsX = CCHeader.NumThreads[0];
