Strip [[clang::optnone]] attribute from Metal shaders due to performance hit with WPO on M1 Macs.

#rnx

diff --git a/Engine/Shaders/Public/Platform/Metal/MetalCommon.ush b/Engine/Shaders/Public/Platform/Metal/MetalCommon.ush
--- a/Engine/Shaders/Public/Platform/Metal/MetalCommon.ush
+++ b/Engine/Shaders/Public/Platform/Metal/MetalCommon.ush
@@ -5,7 +5,7 @@
 =============================================================================*/
 
 // Update this GUID to improve shader recompilation for Metal only shaders
-// GUID = 35B87765-455B-4833-8A78-0804AD8818E7
+// GUID = 389C58BF-880E-4DF0-ACB6-DDB9A8E6F678
 
 #pragma once
 
diff --git a/Engine/Source/Developer/Apple/MetalShaderFormat/Private/MetalDerivedData.cpp b/Engine/Source/Developer/Apple/MetalShaderFormat/Private/MetalDerivedData.cpp
--- a/Engine/Source/Developer/Apple/MetalShaderFormat/Private/MetalDerivedData.cpp
+++ b/Engine/Source/Developer/Apple/MetalShaderFormat/Private/MetalDerivedData.cpp
@@ -1356,6 +1356,13 @@ bool DoCompileMetalShader(
 	}
 #endif
 
+	// Attribute [[clang::optnone]] causes performance hit with WPO on M1 Macs => replace with empty space
+	const std::string ClangOptNoneString = "[[clang::optnone]]";
+	for (size_t Begin = 0, End = 0; (Begin = MetalSource.find(ClangOptNoneString, End)) != std::string::npos; End = Begin)
+	{
+		MetalSource.replace(Begin, ClangOptNoneString.length(), " ");
+	}
+	
 	if (bDumpDebugInfo && !MetalSource.empty())
 	{
 		DumpDebugShaderText(Input, &MetalSource[0], MetalSource.size(), TEXT("metal"));
