Add -fpreserve-invariance for Vertex shaders for Metal versions above 2.2.  macOS only: Reinstate Metal versions 2.2 and 2.3 to work round what appears to be a M1 position invariance compiler bug.

diff --git a/Engine/Shaders/Public/Platform/Metal/MetalCommon.ush b/Engine/Shaders/Public/Platform/Metal/MetalCommon.ush
--- a/Engine/Shaders/Public/Platform/Metal/MetalCommon.ush
+++ b/Engine/Shaders/Public/Platform/Metal/MetalCommon.ush
@@ -5,7 +5,7 @@
 =============================================================================*/
 
 // Update this GUID to improve shader recompilation for Metal only shaders
-// GUID = 389C58BF-880E-4DF0-ACB6-DDB9A8E6F678
+// GUID = 9321cdd8-e027-46cc-83bd-79a19d0cdfb4
 
 #pragma once
 
diff --git a/Engine/Source/Developer/Apple/MetalShaderFormat/Private/MetalDerivedData.cpp b/Engine/Source/Developer/Apple/MetalShaderFormat/Private/MetalDerivedData.cpp
--- a/Engine/Source/Developer/Apple/MetalShaderFormat/Private/MetalDerivedData.cpp
+++ b/Engine/Source/Developer/Apple/MetalShaderFormat/Private/MetalDerivedData.cpp
@@ -1040,10 +1040,20 @@ bool DoCompileMetalShader(
 					TargetDesc.Version = 20400;
 					break;
 				}
+				case 6:
+				{
+					TargetDesc.Version = 20300;
+					break;
+				}
+				case 5:
+				{
+					TargetDesc.Version = 20200;
+					break;
+				}
 				default:
 				{
-					UE_LOG(LogShaders, Warning, TEXT("Metal Shader Version Unsupported, switching to default 2.4")); //EMacMetalShaderStandard::MacMetalSLStandard_Minimum
-					TargetDesc.Version = 20400;
+					UE_LOG(LogShaders, Warning, TEXT("Metal Shader Version Unsupported, switching to default 2.2")); //EMacMetalShaderStandard::MacMetalSLStandard_Minimum
+					TargetDesc.Version = 20200;
 					break;
 				}
 #else
diff --git a/Engine/Source/Developer/Apple/MetalShaderFormat/Private/MetalDerivedData.h b/Engine/Source/Developer/Apple/MetalShaderFormat/Private/MetalDerivedData.h
--- a/Engine/Source/Developer/Apple/MetalShaderFormat/Private/MetalDerivedData.h
+++ b/Engine/Source/Developer/Apple/MetalShaderFormat/Private/MetalDerivedData.h
@@ -48,6 +48,7 @@ struct FMetalShaderBytecodeJob
 	FString OutputObjectFile;
 	FString CompilerVersion;
 	FString MinOSVersion;
+	FString PreserveInvariance;
 	FString DebugInfo;
 	FString MathMode;
 	FString Standard;
diff --git a/Engine/Source/Developer/Apple/MetalShaderFormat/Private/MetalShaderCompiler.cpp b/Engine/Source/Developer/Apple/MetalShaderFormat/Private/MetalShaderCompiler.cpp
--- a/Engine/Source/Developer/Apple/MetalShaderFormat/Private/MetalShaderCompiler.cpp
+++ b/Engine/Source/Developer/Apple/MetalShaderFormat/Private/MetalShaderCompiler.cpp
@@ -677,6 +677,7 @@ void BuildMetalShaderOutput(
 			Job.OutputObjectFile = AIRFileName;
 			Job.CompilerVersion = CompilerVersionString;
 			Job.MinOSVersion = MinOSVersion;
+			Job.PreserveInvariance = Frequency == SF_Vertex && Version > 5 ? TEXT("-fpreserve-invariance") : TEXT("");
 			Job.DebugInfo = DebugInfo;
 			Job.MathMode = MathMode;
 			Job.Standard = Standard;
@@ -933,6 +934,19 @@ void CompileShader_Metal(const FShaderCompilerInput& _Input,FShaderCompilerOutpu
 			MinOSVersion = TEXT("-mmacosx-version-min=12");
 		}
 		break;
+#if PLATFORM_MAC
+	case 6:
+		StandardVersion = TEXT("2.3");
+		MinOSVersion = TEXT("-mmacosx-version-min=11");
+		break;
+	case 5:
+		// Fall through
+	case 0:
+		StandardVersion = TEXT("2.2");
+		MinOSVersion = TEXT("-mmacosx-version-min=10.15");
+		//EMacMetalShaderStandard::MacMetalSLStandard_Minimum
+		break;
+#else
 	case 0:
 		StandardVersion = TEXT("2.4");
 		if (bAppleTV)
@@ -948,8 +962,8 @@ void CompileShader_Metal(const FShaderCompilerInput& _Input,FShaderCompilerOutpu
 			MinOSVersion = TEXT("-mmacosx-version-min=12");
 		}
 		//EIOSMetalShaderStandard::IOSMetalSLStandard_Minimum
-		//EMacMetalShaderStandard::MacMetalSLStandard_Minimum
 		break;
+#endif
 	default:
 		Output.bSucceeded = false;
 		{
diff --git a/Engine/Source/Developer/Apple/MetalShaderFormat/Private/MetalShaderFormat.cpp b/Engine/Source/Developer/Apple/MetalShaderFormat/Private/MetalShaderFormat.cpp
--- a/Engine/Source/Developer/Apple/MetalShaderFormat/Private/MetalShaderFormat.cpp
+++ b/Engine/Source/Developer/Apple/MetalShaderFormat/Private/MetalShaderFormat.cpp
@@ -935,7 +935,7 @@ bool FMetalCompilerToolchain::CompileMetalShader(FMetalShaderBytecodeJob& Job, F
 	FString IncludeArgs = Job.IncludeDir.Len() ? FString::Printf(TEXT("-I %s"), *Job.IncludeDir) : TEXT("");
 	{
 		// Invoke the metal frontend.
-		FString MetalParams = FString::Printf(TEXT("%s %s %s %s -Wno-null-character -fbracket-depth=1024 %s %s %s %s -o %s"), *Job.MinOSVersion, *Job.DebugInfo, *Job.MathMode, TEXT("-c"), *Job.Standard, *Job.Defines, *IncludeArgs, *LocalInputMetalFilePath, *LocalOutputMetalAIRFilePath);
+		FString MetalParams = FString::Printf(TEXT("%s %s %s %s -Wno-null-character -fbracket-depth=1024 %s %s %s %s %s -o %s"), *Job.MinOSVersion, *Job.PreserveInvariance, *Job.DebugInfo, *Job.MathMode, TEXT("-c"), *Job.Standard, *Job.Defines, *IncludeArgs, *LocalInputMetalFilePath, *LocalOutputMetalAIRFilePath);
 		bool bSuccess = this->ExecMetalFrontend(SDK, *MetalParams, &Job.ReturnCode, &Job.Results, &Job.Errors);
 
 		if (!bSuccess || (Job.ReturnCode != 0))
diff --git a/Engine/Source/Developer/Mac/MacTargetPlatform/Classes/MacTargetSettings.h b/Engine/Source/Developer/Mac/MacTargetPlatform/Classes/MacTargetSettings.h
--- a/Engine/Source/Developer/Mac/MacTargetPlatform/Classes/MacTargetSettings.h
+++ b/Engine/Source/Developer/Mac/MacTargetPlatform/Classes/MacTargetSettings.h
@@ -13,8 +13,12 @@
 UENUM()
 enum class EMacMetalShaderStandard : uint8
 {
-    /** Metal Shader 2.4 is the minimum as of UE5.1*/
-    MacMetalSLStandard_Minimum = 0 UMETA(DisplayName="Minimum, Currently v2.4 (12+)"),
+    /** Metal Shader 2.2 is the minimum*/
+    MacMetalSLStandard_Minimum = 0 UMETA(DisplayName="Minimum, Currently v2.2 (10.15+)"),
+    /** Metal Shaders Compatible With macOS 10.15 or later (std=osx-metal2.2) */
+    MacMetalSLStandard_2_2 = 5 UMETA(DisplayName="Metal v2.2 (10.15+)"),
+    /** Metal Shaders Compatible With macOS 11.0 or later (std=osx-metal2.3) */
+    MacMetalSLStandard_2_3 = 6 UMETA(DisplayName="Metal v2.3 (11.0+)"),
     /** Metal Shaders Compatible With macOS 12.0 or later (std=osx-metal2.4) */
     MacMetalSLStandard_2_4 = 7 UMETA(DisplayName="Metal v2.4 (12.0+)"),
     /** Metal Shaders Compatible With macOS 13.0 or later (std=metal3.0) */
diff --git a/Engine/Source/Editor/DetailCustomizations/Private/MacTargetSettingsDetails.cpp b/Engine/Source/Editor/DetailCustomizations/Private/MacTargetSettingsDetails.cpp
--- a/Engine/Source/Editor/DetailCustomizations/Private/MacTargetSettingsDetails.cpp
+++ b/Engine/Source/Editor/DetailCustomizations/Private/MacTargetSettingsDetails.cpp
@@ -394,11 +394,11 @@ void FMacTargetSettingsDetails::SetShaderStandard(int32 Value)
     }
     
     ShaderVersionWarningTextBox->SetError(TEXT(""));
-    if (Value < 7 && Value != 0) // EMacMetalShaderStandard::MacMetalSLStandard_Minimum
+    if (Value < 5 && Value != 0) // EMacMetalShaderStandard::MacMetalSLStandard_Minimum
     {
         FString EngineIdentifier = FEngineVersion::Current().ToString(EVersionComponent::Minor);
         
-        ShaderVersionWarningTextBox->SetError(FString::Printf(TEXT("Minimum Metal Version is 2.4 in UE %s"), *EngineIdentifier));
+        ShaderVersionWarningTextBox->SetError(FString::Printf(TEXT("Minimum Metal Version is 2.2 in UE %s"), *EngineIdentifier));
     }
 }
 
diff --git a/Engine/Source/Runtime/Apple/MetalRHI/Private/MetalShaders.cpp b/Engine/Source/Runtime/Apple/MetalRHI/Private/MetalShaders.cpp
--- a/Engine/Source/Runtime/Apple/MetalRHI/Private/MetalShaders.cpp
+++ b/Engine/Source/Runtime/Apple/MetalRHI/Private/MetalShaders.cpp
@@ -58,15 +58,20 @@ mtlpp::LanguageVersion ValidateVersion(uint32 Version)
         case 7:
             Result = mtlpp::LanguageVersion::Version2_4;
             break;
+		 case 6:
+            Result = mtlpp::LanguageVersion::Version2_3;
+            break;
+		case 5:
+			// Fall through
         case 0:
             Version = 7;
-            Result = mtlpp::LanguageVersion::Version2_4; // minimum version as of UE5.1
+            Result = mtlpp::LanguageVersion::Version2_2; // minimum version as of UE5.1
             break;
         default:
-            //EMacMetalShaderStandard::MacMetalSLStandard_Minimum is currently 2.4
+            //EMacMetalShaderStandard::MacMetalSLStandard_Minimum is currently 2.2
             UE_LOG(LogTemp, Warning, TEXT("The Metal version currently set is not supported anymore. Set it in the Project Settings. Defaulting to the minimum version."));
-            Version = 7;
-            Result = mtlpp::LanguageVersion::Version2_4;
+            Version = 5;
+            Result = mtlpp::LanguageVersion::Version2_2;
             break;
     }
 #else
diff --git a/Engine/Source/Runtime/Apple/MetalRHI/Private/Shaders/Types/Templates/MetalBaseShader.h b/Engine/Source/Runtime/Apple/MetalRHI/Private/Shaders/Types/Templates/MetalBaseShader.h
--- a/Engine/Source/Runtime/Apple/MetalRHI/Private/Shaders/Types/Templates/MetalBaseShader.h
+++ b/Engine/Source/Runtime/Apple/MetalRHI/Private/Shaders/Types/Templates/MetalBaseShader.h
@@ -309,16 +309,34 @@ void TMetalBaseShader<BaseResourceType, ShaderType>::Init(TArrayView<const uint8
 				case 7:
 					MetalVersion = mtlpp::LanguageVersion::Version2_4;
 					break;
-                case 0:
+#if PLATFORM_MAC
+				case 6:
+					MetalVersion = mtlpp::LanguageVersion::Version2_3;
+					break;
+				case 5:
+					// Fall through
+				case 0:
+					// Fall through
+				default:
+					MetalVersion = mtlpp::LanguageVersion::Version2_2;
+					break;
+#else
+				case 0:
                     MetalVersion = mtlpp::LanguageVersion::Version2_4;
                     break;
 				default:
 					UE_LOG(LogRHI, Fatal, TEXT("Failed to create shader with unknown version %d: %s"), Header.Version, *FString(NewShaderString));
 					MetalVersion = mtlpp::LanguageVersion::Version2_4;
 					break;
+#endif
 			}
 			CompileOptions.SetLanguageVersion(MetalVersion);
 
+			if(ShaderType == SF_Vertex && MetalVersion > mtlpp::LanguageVersion::Version2_2)
+			{
+				[CompileOptions.GetPtr() setPreserveInvariance:YES];
+			}
+
 			ns::AutoReleasedError Error;
 			Library = GetMetalDeviceContext().GetDevice().NewLibrary(NewShaderString, CompileOptions, &Error);
 			if (Library == nil)
