Fixing regression with metal shader versions bumps. Temporarily roll back MacOS to 2.2/2.3.

diff --git a/Engine/Source/Developer/Apple/MetalShaderFormat/Private/MetalShaderCompiler.cpp b/Engine/Source/Developer/Apple/MetalShaderFormat/Private/MetalShaderCompiler.cpp
--- a/Engine/Source/Developer/Apple/MetalShaderFormat/Private/MetalShaderCompiler.cpp
+++ b/Engine/Source/Developer/Apple/MetalShaderFormat/Private/MetalShaderCompiler.cpp
@@ -934,19 +934,27 @@ void CompileShader_Metal(const FShaderCompilerInput& _Input,FShaderCompilerOutpu
 			MinOSVersion = TEXT("-mmacosx-version-min=12");
 		}
 		break;
-#if PLATFORM_MAC
 	case 6:
-		StandardVersion = TEXT("2.3");
-		MinOSVersion = TEXT("-mmacosx-version-min=11");
+		StandardVersion = TEXT("2.4");
+		if (bAppleTV)
+		{
+			MinOSVersion = TEXT("-mtvos-version-min=15.0");
+		}
+		else if (bIsMobile)
+		{
+			MinOSVersion = TEXT("-mios-version-min=15.0");
+		}
+		else
+		{
+			// TODO - This is a workaround for an issue with the Apple Shader Compiler
+			// leading to corruption on M1/AMD when > 2.3 versions are used. 
+			// This should be bumped to 2.4 after it's resolved
+			StandardVersion = TEXT("2.3");
+			MinOSVersion = TEXT("-mmacosx-version-min=11");
+		}
 		break;
 	case 5:
 		// Fall through
-	case 0:
-		StandardVersion = TEXT("2.2");
-		MinOSVersion = TEXT("-mmacosx-version-min=10.15");
-		//EMacMetalShaderStandard::MacMetalSLStandard_Minimum
-		break;
-#else
 	case 0:
 		StandardVersion = TEXT("2.4");
 		if (bAppleTV)
@@ -959,11 +967,14 @@ void CompileShader_Metal(const FShaderCompilerInput& _Input,FShaderCompilerOutpu
 		}
 		else
 		{
-			MinOSVersion = TEXT("-mmacosx-version-min=12");
+			// TODO - This is a workaround for an issue with the Apple Shader Compiler
+			// leading to corruption on M1/AMD when > 2.3 versions are used. 
+			// This should be bumped to 2.4 after it's resolved
+			StandardVersion = TEXT("2.2");
+			MinOSVersion = TEXT("-mmacosx-version-min=10.15");
 		}
 		//EIOSMetalShaderStandard::IOSMetalSLStandard_Minimum
 		break;
-#endif
 	default:
 		Output.bSucceeded = false;
 		{
