Various misc fixes for Mac.
* Sets the Min/Max Wave Size
* Handles Vulkan headers and Mac a bit better
* Handles missing [shader] toolchains a bit better.

diff --git a/Engine/Config/Mac/DataDrivenPlatformInfo.ini b/Engine/Config/Mac/DataDrivenPlatformInfo.ini
--- a/Engine/Config/Mac/DataDrivenPlatformInfo.ini
+++ b/Engine/Config/Mac/DataDrivenPlatformInfo.ini
@@ -36,6 +36,8 @@ bSupportsGen5TemporalAA=true
 bSupportsLumenGI=true
 bSupportsSSDIndirect=true
 bSupportsWaveOperations=RuntimeDependent
+MinimumWaveSize=32
+MaximumWaveSize=64
 bSupportsGPUScene=true
 bSupportsFFTBloom=true
 bSupportsDiaphragmDOF=true
diff --git a/Engine/Source/Developer/Apple/MetalShaderFormat/Private/MetalShaderFormat.cpp b/Engine/Source/Developer/Apple/MetalShaderFormat/Private/MetalShaderFormat.cpp
--- a/Engine/Source/Developer/Apple/MetalShaderFormat/Private/MetalShaderFormat.cpp
+++ b/Engine/Source/Developer/Apple/MetalShaderFormat/Private/MetalShaderFormat.cpp
@@ -661,6 +661,10 @@ FMetalCompilerToolchain::EMetalToolchainStatus FMetalCompilerToolchain::FetchCom
 		// But the underlying (windows) implementation of CreateProc puts everything into one pipe, which is written to StdOut.
 		bool bResult = this->ExecMetalFrontend(AppleSDKMac, TEXT("-v --target=air64-apple-darwin18.7.0"), &ReturnCode, &StdOut, &StdOut);
 		check(bResult);
+		if (ReturnCode > 0)
+		{
+			return EMetalToolchainStatus::CouldNotParseCompilerVersion;
+		}
 		
 		Result = ParseCompilerVersionAndTarget(StdOut, this->MetalCompilerVersionString[AppleSDKMac], this->MetalCompilerVersion[AppleSDKMac], this->MetalTargetVersion[AppleSDKMac]);
 
@@ -676,6 +680,10 @@ FMetalCompilerToolchain::EMetalToolchainStatus FMetalCompilerToolchain::FetchCom
 		// metal -v writes its output to stderr
 		bool bResult = this->ExecMetalFrontend(AppleSDKMobile, TEXT("-v --target=air64-apple-darwin18.7.0"), &ReturnCode, &StdOut, &StdOut);
 		check(bResult);
+		if (ReturnCode > 0)
+		{
+			return EMetalToolchainStatus::CouldNotParseCompilerVersion;
+		}
 		
 		Result = ParseCompilerVersionAndTarget(StdOut, this->MetalCompilerVersionString[AppleSDKMobile], this->MetalCompilerVersion[AppleSDKMobile], this->MetalTargetVersion[AppleSDKMobile]);
 	}
@@ -760,9 +768,17 @@ FMetalCompilerToolchain::EMetalToolchainStatus FMetalCompilerToolchain::DoMacNat
 	FString StdOut, StdErr;
 	bool bSuccess = this->ExecGenericCommand(*XcrunPath, *FString::Printf(TEXT("--sdk %s --find %s"), *this->MetalMacSDK, *this->MetalFrontendBinary), &ReturnCode, &StdOut, &StdErr);
 	bSuccess |= FPaths::FileExists(StdOut);
+	if(!bSuccess || ReturnCode > 0)
+	{
+		UE_LOG(LogMetalCompilerSetup, Warning, TEXT("Missing Mac Metal toolchain (macos SDK not found)."));
+		return EMetalToolchainStatus::ToolchainNotFound;
+	}
 	
+	bSuccess = this->ExecGenericCommand(*XcrunPath, *FString::Printf(TEXT("--sdk %s --find %s"), *this->MetalMobileSDK, *this->MetalFrontendBinary), &ReturnCode, &StdOut, &StdErr);
+	bSuccess |= FPaths::FileExists(StdOut);
 	if(!bSuccess || ReturnCode > 0)
 	{
+		UE_LOG(LogMetalCompilerSetup, Warning, TEXT("Missing Mobile Metal toolchain (iphoneos SDK not found)."));
 		return EMetalToolchainStatus::ToolchainNotFound;
 	}
 
diff --git a/Engine/Source/Runtime/VulkanRHI/VulkanRHI.Build.cs b/Engine/Source/Runtime/VulkanRHI/VulkanRHI.Build.cs
--- a/Engine/Source/Runtime/VulkanRHI/VulkanRHI.Build.cs
+++ b/Engine/Source/Runtime/VulkanRHI/VulkanRHI.Build.cs
@@ -71,7 +71,8 @@ public class VulkanRHI : ModuleRules
 			AddEngineThirdPartyPrivateStaticDependencies(Target, "NVAftermath");
 		}
 
-		if (Target.Platform.IsInGroup(UnrealPlatformGroup.Windows) || Target.Platform == UnrealTargetPlatform.Android || Target.IsInPlatformGroup(UnrealPlatformGroup.Unix))
+		if (Target.Platform.IsInGroup(UnrealPlatformGroup.Windows) || Target.Platform == UnrealTargetPlatform.Android
+			|| Target.IsInPlatformGroup(UnrealPlatformGroup.Unix) || Target.Platform == UnrealTargetPlatform.Mac)
 		{
             AddEngineThirdPartyPrivateStaticDependencies(Target, "Vulkan");
         }
@@ -81,29 +82,6 @@ public class VulkanRHI : ModuleRules
 			PrivateDependencyModuleNames.Add("ApplicationCore");
 			AddEngineThirdPartyPrivateStaticDependencies(Target, "SDL2");
 		}
-        else if (Target.Platform == UnrealTargetPlatform.Mac)
-        {
-			string VulkanSDKPath = Environment.GetEnvironmentVariable("VULKAN_SDK");
-
-			bool bHaveVulkan = false;
-			if (!String.IsNullOrEmpty(VulkanSDKPath))
-			{
-				bHaveVulkan = true;
-				PrivateIncludePaths.Add(VulkanSDKPath + "/Include");
-			}
-
-			if (bHaveVulkan)
-			{
-				if (Target.Configuration != UnrealTargetConfiguration.Shipping)
-				{
-					PrivateIncludePathModuleNames.Add("TaskGraph");
-				}
-			}
-			else
-			{
-				PrecompileForTargets = PrecompileTargetsType.None;
-			}
-		}
 		else
 		{
 			PrecompileForTargets = PrecompileTargetsType.None;
