Fix a crash in FShaderCode::GetWriteAccess when r.Shaders.ExtraData=1 on mac.

diff --git a/Engine/Source/Developer/Apple/MetalShaderFormat/Private/MetalShaderCompiler.cpp b/Engine/Source/Developer/Apple/MetalShaderFormat/Private/MetalShaderCompiler.cpp
--- a/Engine/Source/Developer/Apple/MetalShaderFormat/Private/MetalShaderCompiler.cpp
+++ b/Engine/Source/Developer/Apple/MetalShaderFormat/Private/MetalShaderCompiler.cpp
@@ -533,10 +533,6 @@ void BuildMetalShaderOutput(
 	}
 
 	FString MetalCode = FString(USFSource);
-	if (ShaderInput.Environment.CompilerFlags.Contains(CFLAG_ExtraShaderData))
-	{
-		ShaderOutput.ShaderCode.AddOptionalData(FShaderCodeName::Key, TCHAR_TO_UTF8(*ShaderInput.GenerateShaderName()));
-	}
 
 	if (ShaderInput.Environment.CompilerFlags.Contains(CFLAG_Debug))
 	{
