Reducing metal version check in shader cache key to major version only, this will reduce shader compiling time on user machines with a different Xcode version.
Repurpose XcodeVersionInShaderVersion into UseFullMetalVersionInShaderVersion, setting this to true will force shaders to rebuild for each new Xcode version (reverting to previous behaviour)

diff --git a/Engine/Source/Developer/Apple/MetalShaderFormat/Private/MetalShaderFormat.cpp b/Engine/Source/Developer/Apple/MetalShaderFormat/Private/MetalShaderFormat.cpp
--- a/Engine/Source/Developer/Apple/MetalShaderFormat/Private/MetalShaderFormat.cpp
+++ b/Engine/Source/Developer/Apple/MetalShaderFormat/Private/MetalShaderFormat.cpp
@@ -286,47 +286,31 @@ uint32 GetMetalFormatVersion(FName Format)
 		return 0;
 	}
 	
-	// Include the Xcode version when the .ini settings instruct us to do so.
-	bool bAddXcodeVersionInShaderVersion = false;
+	bool bUseFullMetalVersion = false;
 	EShaderPlatform ShaderPlatform = FMetalCompilerToolchain::MetalShaderFormatToLegacyShaderPlatform(Format);
 
 	if(FMetalCompilerToolchain::IsMobile(ShaderPlatform))
 	{
-		GConfig->GetBool(TEXT("/Script/IOSRuntimeSettings.IOSRuntimeSettings"), TEXT("XcodeVersionInShaderVersion"), bAddXcodeVersionInShaderVersion, GEngineIni);
+		GConfig->GetBool(TEXT("/Script/IOSRuntimeSettings.IOSRuntimeSettings"), TEXT("UseFullMetalVersionInShaderVersion"), bUseFullMetalVersion, GEngineIni);
 	}
 	else
 	{
-		GConfig->GetBool(TEXT("/Script/MacTargetPlatform.MacTargetSettings"), TEXT("XcodeVersionInShaderVersion"), bAddXcodeVersionInShaderVersion, GEngineIni);
+		GConfig->GetBool(TEXT("/Script/MacTargetPlatform.MacTargetSettings"), TEXT("UseFullMetalVersionInShaderVersion"), bUseFullMetalVersion, GEngineIni);
 	}
-
-	// We are going to use the LLVM target version instead of the Xcode build, since standalone metal toolchains do not require xcode
-	FMetalCompilerToolchain::PackedVersion TargetVersion = FMetalCompilerToolchain::Get()->GetTargetVersion(ShaderPlatform);
-
-	const FString& CompilerVersionString = FMetalCompilerToolchain::Get()->GetCompilerVersionString(ShaderPlatform);
-	union HashMe
-	{
-		struct
-		{
-			uint16 top;
-			uint16 bottom;
-		};
-		uint32 Value;
-	};
-	HashMe V;
-	V.Value = GetTypeHash(CompilerVersionString);
-	uint16 HashValue = V.top ^ V.bottom;
 	
-	if (!FApp::IsEngineInstalled() && bAddXcodeVersionInShaderVersion)
+	FMetalCompilerToolchain::PackedVersion MetalVersionNumber = FMetalCompilerToolchain::Get()->GetCompilerVersion(ShaderPlatform);
+	uint16 HashValue = MetalVersionNumber.Major;
+	
+	if (bUseFullMetalVersion)
 	{
-		// For local development we'll mix in the LLVM target version.
-		HashValue ^= TargetVersion.Major;
-		HashValue ^= TargetVersion.Minor;
-		HashValue ^= TargetVersion.Patch;
+		// Use entire Metal version if .ini settings instruct us to do so (e.g. p4 dev build)
+		HashValue ^= MetalVersionNumber.Minor;
+		HashValue ^= MetalVersionNumber.Patch;
 	}
 	else
 	{
-		// In the other case (ie, shipping editor binary distributions)
-		// We will only mix the hash of the version string
+		// Only use Metal major version (e.g. Installed build)
+		// Since Metal minor/patch version changes every Xcode minor version, we don't want users to rebuild shaders for every minor version update
 	}
 
 	Version.Version.XcodeVersion = HashValue;
