Update Apple metal version string format to match latest

diff --git a/Engine/Source/Developer/Apple/MetalShaderFormat/Private/MetalShaderFormat.cpp b/Engine/Source/Developer/Apple/MetalShaderFormat/Private/MetalShaderFormat.cpp
--- a/Engine/Source/Developer/Apple/MetalShaderFormat/Private/MetalShaderFormat.cpp
+++ b/Engine/Source/Developer/Apple/MetalShaderFormat/Private/MetalShaderFormat.cpp
@@ -398,18 +398,29 @@ static FMetalCompilerToolchain::EMetalToolchainStatus ParseCompilerVersionAndTar
 		FString& Version = Lines[VersionLineIndex];
 		check(!Version.IsEmpty());
 
-		int32 Major = 0, Minor = 0, Patch = 0;
+		int32 Major = 0, Minor = 0;
 		int32 NumResults = 0;
 #if !PLATFORM_WINDOWS
 		char AppleToolName[PATH_MAX] = { '\0' };
-		NumResults = sscanf(TCHAR_TO_ANSI(*Version), "Apple %s version %d.%d.%d", AppleToolName, &Major, &Minor, &Patch);
+		char SupplementaryVersionName[PATH_MAX] = { '\0' };
+		NumResults = sscanf(TCHAR_TO_ANSI(*Version), "Apple %s version %d.%d (metalfe-%s)", AppleToolName, &Major, &Minor, SupplementaryVersionName);
 #else
 		TCHAR AppleToolName[WINDOWS_MAX_PATH] = { '\0' };
-		NumResults = swscanf_s(*Version, TEXT("Apple %ls version %d.%d.%d"), AppleToolName, WINDOWS_MAX_PATH, &Major, &Minor, &Patch);
+		TCHAR SupplementaryVersionName[WINDOWS_MAX_PATH] = { '\0' };
+		NumResults = swscanf_s(*Version, TEXT("Apple %ls version %d.%d (metalfe-%ls)"), AppleToolName, WINDOWS_MAX_PATH, &Major, &Minor, SupplementaryVersionName, WINDOWS_MAX_PATH);
 #endif
+		if (NumResults != 4)
+		{
+			UE_LOG(LogMetalCompilerSetup, Warning, TEXT("Metal version string format unrecoginzed"));
+			UE_LOG(LogMetalCompilerSetup, Warning, TEXT("Expecting: Apple LLVM version 902.11 (metalfe-902.11.1)"));
+			UE_LOG(LogMetalCompilerSetup, Warning, TEXT("Obtained: %s"), *Version);
+		}
+		
 		PackedVersionNumber.Major = Major;
 		PackedVersionNumber.Minor = Minor;
-		PackedVersionNumber.Patch = Patch;
+		// The version name in brackets is too irregular to extract a useful patch version
+		// Sometimes (metalfe-31001.667.2), sometimes (metalfe-31001.643.2.1), sometimes (metalfe-31001.362-windows)
+		PackedVersionNumber.Patch = 0;
 	}
 
 	if (PackedVersionNumber.Version == 0)
