Fixed: Flickering artifacts occur in the Mobile Game Project on iOS.
Make sure iOS vertex shaders are compiled with PreserveInvariance option set

diff --git a/Engine/Shaders/Private/MobileBasePassVertexShader.usf b/Engine/Shaders/Private/MobileBasePassVertexShader.usf
--- a/Engine/Shaders/Private/MobileBasePassVertexShader.usf
+++ b/Engine/Shaders/Private/MobileBasePassVertexShader.usf
@@ -74,7 +74,7 @@ void Main(
 	WorldPosition.xyz += WorldPositionOffset;
 
 	float4 RasterizedWorldPosition = VertexFactoryGetRasterizedWorldPosition(Input, VFIntermediates, WorldPosition);
-	Output.Position = mul(RasterizedWorldPosition, ResolvedView.TranslatedWorldToClip);
+	Output.Position = INVARIANT(mul(RasterizedWorldPosition, ResolvedView.TranslatedWorldToClip));
 	Output.BasePassInterpolants.PixelPosition = WorldPosition;
 
 #if USE_WORLD_POSITION_EXCLUDING_SHADER_OFFSETS
diff --git a/Engine/Shaders/Public/Platform/Metal/MetalCommon.ush b/Engine/Shaders/Public/Platform/Metal/MetalCommon.ush
--- a/Engine/Shaders/Public/Platform/Metal/MetalCommon.ush
+++ b/Engine/Shaders/Public/Platform/Metal/MetalCommon.ush
@@ -5,7 +5,7 @@
 =============================================================================*/
 
 // Update this GUID to improve shader recompilation for Metal only shaders
-// GUID = F1C3BA6E-2148-457F-BA3C-80F387FF83C0
+// GUID = C53FCB9D-55EE-439A-BD5C-EFB947FEF98E
 
 #pragma once
 
diff --git a/Engine/Source/Developer/Apple/MetalShaderFormat/Private/MetalShaderCompiler.cpp b/Engine/Source/Developer/Apple/MetalShaderFormat/Private/MetalShaderCompiler.cpp
--- a/Engine/Source/Developer/Apple/MetalShaderFormat/Private/MetalShaderCompiler.cpp
+++ b/Engine/Source/Developer/Apple/MetalShaderFormat/Private/MetalShaderCompiler.cpp
@@ -685,6 +685,9 @@ void BuildMetalShaderOutput(
 					IFileManager::Get().Delete(*SaveFile);
 				}
 			}
+
+			// For iOS 14.0+ this is required. Version==0 is IOSMetalSLStandard_Minimum
+			const bool bPreserveInvariance = Frequency == SF_Vertex && (Version == 0 || Version > 5);
 			
 			// TODO This is the actual MetalSL -> AIR piece
 			FMetalShaderBytecodeJob Job;
@@ -697,7 +700,7 @@ void BuildMetalShaderOutput(
 			Job.OutputObjectFile = AIRFileName;
 			Job.CompilerVersion = CompilerVersionString;
 			Job.MinOSVersion = MinOSVersion;
-			Job.PreserveInvariance = Frequency == SF_Vertex && Version > 5 ? TEXT("-fpreserve-invariance") : TEXT("");
+			Job.PreserveInvariance = bPreserveInvariance ? TEXT("-fpreserve-invariance") : TEXT("");
 			Job.DebugInfo = DebugInfo;
 			Job.MathMode = MathMode;
 			Job.Standard = Standard;
