[shaders] deprecate FShaderCompilerEnvironment::RemoteServerData field; since backend-specific remote compile functionality was removed (see https://p4-swarm.epicgames.net/changes/13777670) this is only used for adding unnecessary ifdef'd out debug info to associated shaders.

diff --git a/Engine/Source/Developer/Apple/MetalShaderFormat/Private/MetalShaderCompiler.cpp b/Engine/Source/Developer/Apple/MetalShaderFormat/Private/MetalShaderCompiler.cpp
--- a/Engine/Source/Developer/Apple/MetalShaderFormat/Private/MetalShaderCompiler.cpp
+++ b/Engine/Source/Developer/Apple/MetalShaderFormat/Private/MetalShaderCompiler.cpp
@@ -793,58 +793,6 @@ void BuildMetalShaderOutput(
 	External interface.
 ------------------------------------------------------------------------------*/
 
-FString CreateRemoteDataFromEnvironment(const FShaderCompilerEnvironment& Environment)
-{
-	FString Line = TEXT("\n#if 0 /*BEGIN_REMOTE_SERVER*/\n");
-	for (auto Pair : Environment.RemoteServerData)
-	{
-		Line += FString::Printf(TEXT("%s=%s\n"), *Pair.Key, *Pair.Value);
-	}
-	Line += TEXT("#endif /*END_REMOTE_SERVER*/\n");
-	return Line;
-}
-
-void CreateEnvironmentFromRemoteData(const FString& String, FShaderCompilerEnvironment& OutEnvironment)
-{
-	FString Prolog = TEXT("#if 0 /*BEGIN_REMOTE_SERVER*/");
-	int32 FoundBegin = String.Find(Prolog, ESearchCase::CaseSensitive);
-	if (FoundBegin == INDEX_NONE)
-	{
-		return;
-	}
-	int32 FoundEnd = String.Find(TEXT("#endif /*END_REMOTE_SERVER*/"), ESearchCase::CaseSensitive, ESearchDir::FromStart, FoundBegin);
-	if (FoundEnd == INDEX_NONE)
-	{
-		return;
-	}
-
-	// +1 for EOL
-	const TCHAR* Ptr = &String[FoundBegin + 1 + Prolog.Len()];
-	const TCHAR* PtrEnd = &String[FoundEnd];
-	while (Ptr < PtrEnd)
-	{
-		FString Key;
-		if (!CrossCompiler::ParseIdentifier(Ptr, Key))
-		{
-			return;
-		}
-		if (!CrossCompiler::Match(Ptr, TEXT("=")))
-		{
-			return;
-		}
-		FString Value;
-		if (!CrossCompiler::ParseString(Ptr, Value))
-		{
-			return;
-		}
-		if (!CrossCompiler::Match(Ptr, '\n'))
-		{
-			return;
-		}
-		OutEnvironment.RemoteServerData.FindOrAdd(Key) = Value;
-	}
-}
-
 void CompileShader_Metal(const FShaderCompilerInput& _Input,FShaderCompilerOutput& Output,const FString& WorkingDirectory)
 {
 	auto Input = _Input;
@@ -1062,7 +1010,6 @@ void CompileShader_Metal(const FShaderCompilerInput& _Input,FShaderCompilerOutpu
 
 		// Remove const as we are on debug-only mode
 		CrossCompiler::CreateEnvironmentFromResourceTable(PreprocessedShader, (FShaderCompilerEnvironment&)Input.Environment);
-		CreateEnvironmentFromRemoteData(PreprocessedShader, (FShaderCompilerEnvironment&)Input.Environment);
 	}
 	else
 	{
@@ -1126,14 +1073,6 @@ void CompileShader_Metal(const FShaderCompilerInput& _Input,FShaderCompilerOutpu
 	{
 		UE::ShaderCompilerCommon::FDebugShaderDataOptions DebugDataOptions;
 		DebugDataOptions.HlslCCFlags = CCFlags;
-		DebugDataOptions.AppendPostSource = [&Input]()
-		{
-			// add the remote data if necessary
-//			if (IsRemoteBuildingConfigured(&Input.Environment))
-			{
-				return CreateRemoteDataFromEnvironment(Input.Environment);
-			}
-		};
 		UE::ShaderCompilerCommon::DumpDebugShaderData(Input, PreprocessedShader, DebugDataOptions);
 	}
 
diff --git a/Engine/Source/Runtime/Engine/Private/ShaderCompiler/ShaderCompiler.cpp b/Engine/Source/Runtime/Engine/Private/ShaderCompiler/ShaderCompiler.cpp
--- a/Engine/Source/Runtime/Engine/Private/ShaderCompiler/ShaderCompiler.cpp
+++ b/Engine/Source/Runtime/Engine/Private/ShaderCompiler/ShaderCompiler.cpp
@@ -6272,22 +6272,6 @@ void GlobalBeginCompileShader(
 		{
 			Input.Environment.CompilerFlags.Add(CFLAG_Debug);
 		}
-		else
-		{
-			// populate the data in the shader input environment
-			FString RemoteServer;
-			FString UserName;
-			FString SSHKey;
-			GConfig->GetString(TEXT("/Script/IOSRuntimeSettings.IOSRuntimeSettings"), TEXT("RemoteServerName"), RemoteServer, GEngineIni);
-			GConfig->GetString(TEXT("/Script/IOSRuntimeSettings.IOSRuntimeSettings"), TEXT("RSyncUsername"), UserName, GEngineIni);
-			GConfig->GetString(TEXT("/Script/IOSRuntimeSettings.IOSRuntimeSettings"), TEXT("SSHPrivateKeyOverridePath"), SSHKey, GEngineIni);
-			Input.Environment.RemoteServerData.Add(TEXT("RemoteServerName"), RemoteServer);
-			Input.Environment.RemoteServerData.Add(TEXT("RSyncUsername"), UserName);
-			if (SSHKey.Len() > 0)
-			{
-				Input.Environment.RemoteServerData.Add(TEXT("SSHPrivateKeyOverridePath"), SSHKey);
-			}
-		}
 		
 		// Shaders built for archiving - for Metal that requires compiling the code in a different way so that we can strip it later
 		bool bArchive = false;
diff --git a/Engine/Source/Runtime/RenderCore/Public/ShaderCore.h b/Engine/Source/Runtime/RenderCore/Public/ShaderCore.h
--- a/Engine/Source/Runtime/RenderCore/Public/ShaderCore.h
+++ b/Engine/Source/Runtime/RenderCore/Public/ShaderCore.h
@@ -525,6 +525,8 @@ struct FShaderCompilerEnvironment
 	TMap<uint32,uint8> RenderTargetOutputFormatsMap;
 	TMap<FString, FResourceTableEntry> ResourceTableMap;
 	TMap<FString, FUniformBufferEntry> UniformBufferMap;
+
+	UE_DEPRECATED(5.3, "RemoteServerData field is deprecated (no longer used in compilation backends).")
 	TMap<FString, FString> RemoteServerData;
 
 	const ITargetPlatform* TargetPlatform = nullptr;
@@ -545,10 +547,21 @@ struct FShaderCompilerEnvironment
 	{
 	}
 
-	// Used as a baseclasss, make sure we're not incorrectly destroyed through a baseclass pointer
+	// Used as a baseclass, make sure we're not incorrectly destroyed through a baseclass pointer
 	// This will be expensive to destroy anyway, additional vcall overhead should be small
+	PRAGMA_DISABLE_DEPRECATION_WARNINGS
+	// Temporarily disable deprecation warnings for this default destructor triggering due to the 
+	// deprecated RemoteServerData field. When the field is removed warning disable should be as 
+	// well, but the defaulted destructor should remain due to the above comment.
 	virtual ~FShaderCompilerEnvironment() = default;
 
+	// Explicitly default assignment operator and copy constructor operator with warnings disabled
+	// to avoid warnings in implicitly-generated functions due to deprecation of RemoteServerData. 
+	// These can be removed entirely (revert to implicitly-generated) when the field itself is.
+	FShaderCompilerEnvironment(const FShaderCompilerEnvironment&) = default;
+	FShaderCompilerEnvironment& operator=(const FShaderCompilerEnvironment&) = default;
+	PRAGMA_ENABLE_DEPRECATION_WARNINGS
+
 	/**
 	 * Works for TCHAR
 	 * e.g. SetDefine(TEXT("NAME"), TEXT("Test"));
@@ -632,7 +645,6 @@ struct FShaderCompilerEnvironment
 		Ar << RenderTargetOutputFormatsMap;
 		Ar << ResourceTableMap;
 		Ar << UniformBufferMap;
-		Ar << RemoteServerData;
 		Ar << FullPrecisionInPS;
 	}
 
@@ -673,7 +685,6 @@ struct FShaderCompilerEnvironment
 		Definitions.Merge(Other.Definitions);
 		CompileArgs.Append(Other.CompileArgs);
 		RenderTargetOutputFormatsMap.Append(Other.RenderTargetOutputFormatsMap);
-		RemoteServerData.Append(Other.RemoteServerData);
 		FullPrecisionInPS |= Other.FullPrecisionInPS;
 	}
 
