Fixed bug where arrays were flattened when not using SM6, also marked SM6 as experimental
Enabled distance fields for the desktop renderer on iOS

diff --git a/Engine/Config/IOS/DataDrivenPlatformInfo.ini b/Engine/Config/IOS/DataDrivenPlatformInfo.ini
--- a/Engine/Config/IOS/DataDrivenPlatformInfo.ini
+++ b/Engine/Config/IOS/DataDrivenPlatformInfo.ini
@@ -67,4 +67,5 @@ bSupportsIndexBufferUAVs=true
 bSupports4ComponentUAVReadWrite=true
 bSupportsManualVertexFetch=false
 bSupportsVolumeTextureAtomics=false
-FriendlyName=LOCTEXT("FriendlyShaderPlatformName_IOSSM5", "IOS Metal SM5")
\ No newline at end of file
+bSupportsDistanceFields=true
+FriendlyName=LOCTEXT("FriendlyShaderPlatformName_IOSSM5", "IOS Metal SM5")
diff --git a/Engine/Source/Developer/Apple/MetalShaderFormat/Private/MetalDerivedData.cpp b/Engine/Source/Developer/Apple/MetalShaderFormat/Private/MetalDerivedData.cpp
--- a/Engine/Source/Developer/Apple/MetalShaderFormat/Private/MetalDerivedData.cpp
+++ b/Engine/Source/Developer/Apple/MetalShaderFormat/Private/MetalDerivedData.cpp
@@ -982,25 +982,28 @@ bool DoCompileMetalShader(
 			TargetDesc.CompileFlags.SetDefine(TEXT("invariant_float_math"), Options.bEnableFMAPass ? 1 : 0);
 			TargetDesc.CompileFlags.SetDefine(TEXT("enable_decoration_binding"), 1);
 
-    		// Detect if we need to patch VSM shaders (flatten 2D array as regular 2D texture).
-			// Must be done as VSM uses 2DArray and requires atomics support. And Metal does not
-			// support atomics on 2Darray...
-			bool bShouldFlatten2DArray = InPreprocessedShader.Find("VIRTUAL_SHADOW_MAP 1") > 0
-                                        || InPreprocessedShader.Find("VIRTUAL_TEXTURE_TARGET 1") > 0
-                                        || Input.ShaderName.Find("PhysicalPage") != INDEX_NONE
-                                        || Input.ShaderName.Find("Virtual") != INDEX_NONE
-                                        || Input.ShaderName.Find("ClassifyMaterial") != INDEX_NONE;
-
-			// Need to patch Clear/Memset CS too.
-			if (!bShouldFlatten2DArray)
-			{
-			    if (Input.Environment.GetCompileArgument(TEXT("RESOURCE_TYPE"), 0) == 2)
-			    {
-			        bShouldFlatten2DArray = true;
-			    }
-			}
-
-			TargetDesc.CompileFlags.SetDefine(TEXT("flatten_2d_array"), bShouldFlatten2DArray ? 1 : 0);
+            if(Input.Target.Platform == SP_METAL_SM6)
+            {
+                // Detect if we need to patch VSM shaders (flatten 2D array as regular 2D texture).
+                // Must be done as VSM uses 2DArray and requires atomics support. And Metal does not
+                // support atomics on 2Darray...
+                bool bShouldFlatten2DArray = InPreprocessedShader.Find("VIRTUAL_SHADOW_MAP 1") > 0
+                    || InPreprocessedShader.Find("VIRTUAL_TEXTURE_TARGET 1") > 0
+                    || Input.ShaderName.Find("PhysicalPage") != INDEX_NONE
+                    || Input.ShaderName.Find("Virtual") != INDEX_NONE
+                    || Input.ShaderName.Find("ClassifyMaterial") != INDEX_NONE;
+                
+                // Need to patch Clear/Memset CS too.
+                if (!bShouldFlatten2DArray)
+                {
+                    if (Input.Environment.GetCompileArgument(TEXT("RESOURCE_TYPE"), 0) == 2)
+                    {
+                        bShouldFlatten2DArray = true;
+                    }
+                }
+                
+                TargetDesc.CompileFlags.SetDefine(TEXT("flatten_2d_array"), bShouldFlatten2DArray ? 1 : 0);
+            }
 
 			switch (Semantics)
 			{
diff --git a/Engine/Source/Editor/DetailCustomizations/Private/MacTargetSettingsDetails.cpp b/Engine/Source/Editor/DetailCustomizations/Private/MacTargetSettingsDetails.cpp
--- a/Engine/Source/Editor/DetailCustomizations/Private/MacTargetSettingsDetails.cpp
+++ b/Engine/Source/Editor/DetailCustomizations/Private/MacTargetSettingsDetails.cpp
@@ -138,7 +138,7 @@ static FText GetFriendlyNameFromRHINameMac(FName InRHIName)
 		FriendlyRHIName = LOCTEXT("MetalSM5", "Mac Metal Desktop Renderer (SM5, Metal 2.4+, macOS Monterey 12.0 or later)");
 		break;
     case SP_METAL_SM6:
-        FriendlyRHIName = LOCTEXT("MetalSM6", "Mac Metal Desktop Renderer (SM6, Metal 2.4+, macOS Monterey 12.0 or later)");
+        FriendlyRHIName = LOCTEXT("MetalSM6", "Mac Metal Desktop Renderer Experimental (SM6, Metal 2.4+, macOS 14.0 or later, M2+)");
         break;
 	case SP_METAL_MACES3_1:
 		FriendlyRHIName = LOCTEXT("MetalMobile", "Mac Metal High-End Mobile Preview (Mobile Preview)");
