Required Metal shader compiler changes to support iOS Simulator. Set to disabled by default via "bEnableSimulatorSupport" in IOSRuntimeSettings.

#rnx

diff --git a/Engine/Config/BaseEngine.ini b/Engine/Config/BaseEngine.ini
--- a/Engine/Config/BaseEngine.ini
+++ b/Engine/Config/BaseEngine.ini
@@ -2841,6 +2841,7 @@ bEnableDynamicMaxFPS=False
 bSupportsIPad=True
 bSupportsIPhone=True
 bEnableSplitView=False
+bEnableSimulatorSupport=False
 AdditionalPlistData=
 RemoteServerName=
 RemoteServerOverrideBuildPath=
diff --git a/Engine/Source/Developer/Apple/MetalShaderFormat/Private/MetalShaderCompiler.cpp b/Engine/Source/Developer/Apple/MetalShaderFormat/Private/MetalShaderCompiler.cpp
--- a/Engine/Source/Developer/Apple/MetalShaderFormat/Private/MetalShaderCompiler.cpp
+++ b/Engine/Source/Developer/Apple/MetalShaderFormat/Private/MetalShaderCompiler.cpp
@@ -261,7 +261,7 @@ void BuildMetalShaderOutput(
 	const EShaderFrequency Frequency = ShaderOutput.Target.GetFrequency();
 
 	//TODO read from toolchain
-	const bool bIsMobile = (ShaderInput.Target.Platform == SP_METAL || ShaderInput.Target.Platform == SP_METAL_MRT || ShaderInput.Target.Platform == SP_METAL_TVOS || ShaderInput.Target.Platform == SP_METAL_MRT_TVOS);
+	const bool bIsMobile = (ShaderInput.Target.Platform == SP_METAL || ShaderInput.Target.Platform == SP_METAL_MRT || ShaderInput.Target.Platform == SP_METAL_TVOS || ShaderInput.Target.Platform == SP_METAL_MRT_TVOS || ShaderInput.Target.Platform == SP_METAL_SIM);
 	bool bNoFastMath = ShaderInput.Environment.CompilerFlags.Contains(CFLAG_NoFastMath);
 	const bool bUsingWPO = ShaderInput.Environment.GetCompileArgument(TEXT("USES_WORLD_POSITION_OFFSET"), false);
 	if (bUsingWPO && (ShaderInput.Target.Platform == SP_METAL_MRT || ShaderInput.Target.Platform == SP_METAL_MRT_TVOS) && Frequency == SF_Vertex)
@@ -794,10 +794,16 @@ void CompileShader_Metal(const FShaderCompilerInput& _Input,FShaderCompilerOutpu
 
 	// TODO read from toolchain
 	bool bAppleTV = (Input.ShaderFormat == NAME_SF_METAL_TVOS || Input.ShaderFormat == NAME_SF_METAL_MRT_TVOS);
+	bool bIsSimulator = false;
 	if (Input.ShaderFormat == NAME_SF_METAL || Input.ShaderFormat == NAME_SF_METAL_TVOS)
 	{
 		AdditionalDefines.SetDefine(TEXT("METAL_PROFILE"), 1);
 	}
+	else if (Input.ShaderFormat == NAME_SF_METAL_SIM)
+	{
+		AdditionalDefines.SetDefine(TEXT("METAL_PROFILE"), 1);
+		bIsSimulator = true;
+	}
 	else if (Input.ShaderFormat == NAME_SF_METAL_MRT || Input.ShaderFormat == NAME_SF_METAL_MRT_TVOS)
 	{
 		AdditionalDefines.SetDefine(TEXT("METAL_MRT_PROFILE"), 1);
@@ -847,7 +853,14 @@ void CompileShader_Metal(const FShaderCompilerInput& _Input,FShaderCompilerOutpu
         }
         else if (bIsMobile)
         {
-            MinOSVersion = TEXT("-mios-version-min=17.0");
+			if (bIsSimulator)
+			{
+				MinOSVersion = TEXT("-miphonesimulator-version-min=17.0");
+			}
+			else
+			{
+				MinOSVersion = TEXT("-mios-version-min=17.0");
+			}
         }
         else
         {
@@ -862,7 +875,14 @@ void CompileShader_Metal(const FShaderCompilerInput& _Input,FShaderCompilerOutpu
         }
         else if (bIsMobile)
         {
-            MinOSVersion = TEXT("-mios-version-min=16.0");
+			if (bIsSimulator)
+			{
+				MinOSVersion = TEXT("-miphonesimulator-version-min=16.0");
+			}
+			else
+			{
+				MinOSVersion = TEXT("-mios-version-min=16.0");
+			}
         }
         else
         {
@@ -878,7 +898,14 @@ void CompileShader_Metal(const FShaderCompilerInput& _Input,FShaderCompilerOutpu
 		}
 		else if (bIsMobile)
 		{
-			MinOSVersion = TEXT("-mios-version-min=15.0");
+			if (bIsSimulator)
+			{
+				MinOSVersion = TEXT("-miphonesimulator-version-min=15.0");
+			}
+			else
+			{
+				MinOSVersion = TEXT("-mios-version-min=15.0");
+			}
 		}
 		else
 		{
@@ -893,7 +920,14 @@ void CompileShader_Metal(const FShaderCompilerInput& _Input,FShaderCompilerOutpu
 		}
 		else if (bIsMobile)
 		{
-			MinOSVersion = TEXT("-mios-version-min=15.0");
+			if (bIsSimulator)
+			{
+				MinOSVersion = TEXT("-miphonesimulator-version-min=15.0");
+			}
+			else
+			{
+				MinOSVersion = TEXT("-mios-version-min=15.0");
+			}
 		}
 		else
 		{
@@ -914,7 +948,14 @@ void CompileShader_Metal(const FShaderCompilerInput& _Input,FShaderCompilerOutpu
 		}
 		else if (bIsMobile)
 		{
-			MinOSVersion = TEXT("-mios-version-min=15.0");
+			if (bIsSimulator)
+			{
+				MinOSVersion = TEXT("-miphonesimulator-version-min=15.0");
+			}
+			else
+			{
+				MinOSVersion = TEXT("-mios-version-min=15.0");
+			}
 		}
 		else
 		{
diff --git a/Engine/Source/Developer/Apple/MetalShaderFormat/Private/MetalShaderFormat.cpp b/Engine/Source/Developer/Apple/MetalShaderFormat/Private/MetalShaderFormat.cpp
--- a/Engine/Source/Developer/Apple/MetalShaderFormat/Private/MetalShaderFormat.cpp
+++ b/Engine/Source/Developer/Apple/MetalShaderFormat/Private/MetalShaderFormat.cpp
@@ -65,12 +65,13 @@ public:
 		OutFormats.Add(NAME_SF_METAL_MRT_TVOS);
 		OutFormats.Add(NAME_SF_METAL_SM5);
         OutFormats.Add(NAME_SF_METAL_SM6);
+        OutFormats.Add(NAME_SF_METAL_SIM);
 		OutFormats.Add(NAME_SF_METAL_MACES3_1);
 		OutFormats.Add(NAME_SF_METAL_MRT_MAC);
 	}
 	virtual void CompileShader(FName Format, const struct FShaderCompilerInput& Input, struct FShaderCompilerOutput& Output,const FString& WorkingDirectory) const override final
 	{
-		check(Format == NAME_SF_METAL || Format == NAME_SF_METAL_MRT || Format == NAME_SF_METAL_TVOS || Format == NAME_SF_METAL_MRT_TVOS || Format == NAME_SF_METAL_SM5 || Format == NAME_SF_METAL_SM6 || Format == NAME_SF_METAL_MACES3_1 || Format == NAME_SF_METAL_MRT_MAC);
+		check(Format == NAME_SF_METAL || Format == NAME_SF_METAL_MRT || Format == NAME_SF_METAL_TVOS || Format == NAME_SF_METAL_MRT_TVOS || Format == NAME_SF_METAL_SM5 || Format == NAME_SF_METAL_SM6 || Format == NAME_SF_METAL_SIM || Format == NAME_SF_METAL_MACES3_1 || Format == NAME_SF_METAL_MRT_MAC);
 		CompileShader_Metal(Input, Output, WorkingDirectory);
 	}
 	virtual bool CanStripShaderCode(bool const bNativeFormat) const override final
@@ -103,7 +104,7 @@ public:
 		check(Components.Num() == 2);
 		FName ShaderFormatName(Components[0]);
 
-		check(ShaderFormatName == NAME_SF_METAL || ShaderFormatName == NAME_SF_METAL_MRT || ShaderFormatName == NAME_SF_METAL_TVOS || ShaderFormatName == NAME_SF_METAL_MRT_TVOS || ShaderFormatName == NAME_SF_METAL_SM5 || ShaderFormatName == NAME_SF_METAL_SM6 || ShaderFormatName == NAME_SF_METAL_MACES3_1 || ShaderFormatName == NAME_SF_METAL_MRT_MAC);
+		check(ShaderFormatName == NAME_SF_METAL || ShaderFormatName == NAME_SF_METAL_MRT || ShaderFormatName == NAME_SF_METAL_TVOS || ShaderFormatName == NAME_SF_METAL_MRT_TVOS || ShaderFormatName == NAME_SF_METAL_SM5 || ShaderFormatName == NAME_SF_METAL_SM6  || ShaderFormatName == NAME_SF_METAL_SIM || ShaderFormatName == NAME_SF_METAL_MACES3_1 || ShaderFormatName == NAME_SF_METAL_MRT_MAC);
 
 		const FString ArchivePath = (WorkingDirectory / ShaderFormatAndShaderPlatformName.GetPlainNameString());
 		IFileManager::Get().DeleteDirectory(*ArchivePath, false, true);
@@ -529,6 +530,7 @@ EShaderPlatform FMetalCompilerToolchain::MetalShaderFormatToLegacyShaderPlatform
 	if (ShaderFormat == NAME_SF_METAL_MRT_MAC)		return SP_METAL_MRT_MAC;
 	if (ShaderFormat == NAME_SF_METAL_SM5)			return SP_METAL_SM5;
     if (ShaderFormat == NAME_SF_METAL_SM6)          return SP_METAL_SM6;
+    if (ShaderFormat == NAME_SF_METAL_SIM)          return SP_METAL_SIM;
 	if (ShaderFormat == NAME_SF_METAL_MACES3_1)		return SP_METAL_MACES3_1;
 
 	return SP_NumPlatforms;
diff --git a/Engine/Source/Developer/Apple/MetalShaderFormat/Private/MetalShaderFormat.h b/Engine/Source/Developer/Apple/MetalShaderFormat/Private/MetalShaderFormat.h
--- a/Engine/Source/Developer/Apple/MetalShaderFormat/Private/MetalShaderFormat.h
+++ b/Engine/Source/Developer/Apple/MetalShaderFormat/Private/MetalShaderFormat.h
@@ -33,6 +33,7 @@ static FName NAME_SF_METAL_TVOS(TEXT("SF_METAL_TVOS"));
 static FName NAME_SF_METAL_MRT_TVOS(TEXT("SF_METAL_MRT_TVOS"));
 static FName NAME_SF_METAL_SM5(TEXT("SF_METAL_SM5"));
 static FName NAME_SF_METAL_SM6(TEXT("SF_METAL_SM6"));
+static FName NAME_SF_METAL_SIM(TEXT("SF_METAL_SIM"));
 static FName NAME_SF_METAL_MACES3_1(TEXT("SF_METAL_MACES3_1"));
 static FName NAME_SF_METAL_MRT_MAC(TEXT("SF_METAL_MRT_MAC"));
 
@@ -137,7 +138,7 @@ public:
 	// True if Platform indicates a mobile (ios, tvos) platform
 	static bool IsMobile(const EShaderPlatform Platform)
 	{
-		return (Platform == SP_METAL || Platform == SP_METAL_MRT || Platform == SP_METAL_TVOS || Platform == SP_METAL_MRT_TVOS);
+		return (Platform == SP_METAL || Platform == SP_METAL_MRT || Platform == SP_METAL_SIM || Platform == SP_METAL_TVOS || Platform == SP_METAL_MRT_TVOS);
 	}
 
 	// True if Format is the FName of a mobile shader format (ios, tvos)
diff --git a/Engine/Source/Developer/IOS/IOSTargetPlatform/Private/IOSTargetPlatform.cpp b/Engine/Source/Developer/IOS/IOSTargetPlatform/Private/IOSTargetPlatform.cpp
--- a/Engine/Source/Developer/IOS/IOSTargetPlatform/Private/IOSTargetPlatform.cpp
+++ b/Engine/Source/Developer/IOS/IOSTargetPlatform/Private/IOSTargetPlatform.cpp
@@ -517,6 +517,7 @@ bool FIOSTargetPlatform::SupportsFeature( ETargetPlatformFeatures Feature ) cons
 void FIOSTargetPlatform::GetAllPossibleShaderFormats( TArray<FName>& OutFormats ) const
 {
 	static FName NAME_SF_METAL(TEXT("SF_METAL"));
+	static FName NAME_SF_METAL_SIM(TEXT("SF_METAL_SIM"));
 	static FName NAME_SF_METAL_MRT(TEXT("SF_METAL_MRT"));
 	static FName NAME_SF_METAL_TVOS(TEXT("SF_METAL_TVOS"));
 	static FName NAME_SF_METAL_MRT_TVOS(TEXT("SF_METAL_MRT_TVOS"));
@@ -540,6 +541,13 @@ void FIOSTargetPlatform::GetAllPossibleShaderFormats( TArray<FName>& OutFormats
 		if (SupportsMetal())
 		{
 			OutFormats.AddUnique(NAME_SF_METAL);
+
+			bool bEnableSimulatorSupport = false;
+			GConfig->GetBool(TEXT("/Script/IOSRuntimeSettings.IOSRuntimeSettings"), TEXT("bEnableSimulatorSupport"), bEnableSimulatorSupport, GEngineIni);
+			if (bEnableSimulatorSupport)
+			{
+				OutFormats.AddUnique(NAME_SF_METAL_SIM);
+			}
 		}
 
 		if (SupportsMetalMRT())
@@ -663,6 +671,9 @@ void FIOSTargetPlatform::GetTextureFormats( const UTexture* Texture, TArray< TAr
 		}
 	}
 
+	bool bEnableSimulatorSupport = false;
+	GConfig->GetBool(TEXT("/Script/IOSRuntimeSettings.IOSRuntimeSettings"), TEXT("bEnableSimulatorSupport"), bEnableSimulatorSupport, GEngineIni);
+
 	for (FName& TextureFormatName : OutFormats.Last())
 	{
 		if (Texture->GetTextureClass() == ETextureClass::Cube) 
@@ -675,6 +686,14 @@ void FIOSTargetPlatform::GetTextureFormats( const UTexture* Texture, TArray< TAr
 				TextureFormatName = FName(TEXT("ETC2_RGBA"));
 			}
 		}
+
+		// Currently (Xcode14), the iOS Simulator does not support compressed Volume textures.
+		if (bEnableSimulatorSupport && Texture->GetTextureClass() == ETextureClass::Volume)
+		{
+			FTextureFormatSettings FormatSettings;
+			Texture->GetDefaultFormatSettings(FormatSettings);
+			TextureFormatName = FName(TEXT("RGB8"));
+		}
 	}
 }
 
diff --git a/Engine/Source/Programs/UnrealBuildTool/ProjectFiles/Xcode/XcodeProject.cs b/Engine/Source/Programs/UnrealBuildTool/ProjectFiles/Xcode/XcodeProject.cs
--- a/Engine/Source/Programs/UnrealBuildTool/ProjectFiles/Xcode/XcodeProject.cs
+++ b/Engine/Source/Programs/UnrealBuildTool/ProjectFiles/Xcode/XcodeProject.cs
@@ -1689,16 +1689,15 @@ namespace UnrealBuildTool.XcodeProjectXcconfig
 
 				if (Platform == UnrealTargetPlatform.IOS)
 				{
-					// if (TargetRules.Architecture == UnrealArch.IOSSimulator)
-					// {
-					// 	SDKRoot = "iphonesimulator";
-					// 	SupportedPlatforms = "iphonesimulator";
-					// }
-					// else
-					// {	
-						SDKRoot = "iphoneos";
-						SupportedPlatforms = "iphoneos";
-					// }
+					bool bEnableSimulatorSupport = false;
+					PlatformIni.GetBool("/Script/IOSRuntimeSettings.IOSRuntimeSettings", "bEnableSimulatorSupport", out bEnableSimulatorSupport);
+
+					SDKRoot = "iphoneos";
+					SupportedPlatforms = "iphoneos";
+					if (bEnableSimulatorSupport)
+					{
+						SupportedPlatforms += " iphonesimulator";
+					}
 				
 					DeploymentTargetKey = "IPHONEOS_DEPLOYMENT_TARGET";
 					SupportedDevices = UnrealData.IOSProjectSettings!.RuntimeDevices;
diff --git a/Engine/Source/Runtime/IOS/IOSRuntimeSettings/Classes/IOSRuntimeSettings.h b/Engine/Source/Runtime/IOS/IOSRuntimeSettings/Classes/IOSRuntimeSettings.h
--- a/Engine/Source/Runtime/IOS/IOSRuntimeSettings/Classes/IOSRuntimeSettings.h
+++ b/Engine/Source/Runtime/IOS/IOSRuntimeSettings/Classes/IOSRuntimeSettings.h
@@ -270,6 +270,10 @@ public:
     UPROPERTY(GlobalConfig, EditAnywhere, Category = Build, meta = (DisplayName = "Enable iPad Split View"))
     bool bEnableSplitView;
     
+    // Whether or not iOS Simulator support should be enabled for this project (Experimental)
+    UPROPERTY(GlobalConfig, EditAnywhere, Category = Build, meta = (DisplayName = "Enable iOS Simulator Support (Experimental)", ConfigRestartRequired = true))
+    bool bEnableSimulatorSupport;
+    
     /** Set the maximum frame rate to save on power consumption */
     UPROPERTY(GlobalConfig, EditAnywhere, Category = PowerUsage, meta = (ConfigHierarchyEditable))
     EPowerUsageFrameRateLock FrameRateLock;
diff --git a/Engine/Source/Runtime/IOS/IOSRuntimeSettings/Private/IOSRuntimeSettings.cpp b/Engine/Source/Runtime/IOS/IOSRuntimeSettings/Private/IOSRuntimeSettings.cpp
--- a/Engine/Source/Runtime/IOS/IOSRuntimeSettings/Private/IOSRuntimeSettings.cpp
+++ b/Engine/Source/Runtime/IOS/IOSRuntimeSettings/Private/IOSRuntimeSettings.cpp
@@ -32,6 +32,7 @@ UIOSRuntimeSettings::UIOSRuntimeSettings(const FObjectInitializer& ObjectInitial
 	bSupportsIPad = true;
 	bSupportsIPhone = true;
 	bEnableSplitView = false;
+	bEnableSimulatorSupport = false;
 	MinimumiOSVersion = EIOSVersion::IOS_Minimum;
     bBuildAsFramework = true;
 	bGeneratedSYMFile = false;
diff --git a/Engine/Source/Runtime/RHI/Private/Apple/Platform/IOS/IOSPlatformDynamicRHI.h b/Engine/Source/Runtime/RHI/Private/Apple/Platform/IOS/IOSPlatformDynamicRHI.h
--- a/Engine/Source/Runtime/RHI/Private/Apple/Platform/IOS/IOSPlatformDynamicRHI.h
+++ b/Engine/Source/Runtime/RHI/Private/Apple/Platform/IOS/IOSPlatformDynamicRHI.h
@@ -18,13 +18,13 @@
 #if PLATFORM_TVOS
 #define SHADER_PLATFORM_METAL_MRT	SP_METAL_MRT_TVOS
 #define SHADER_PLATFORM_METAL		SP_METAL_TVOS
-
-#else // PLATFORM_TVOS
+#elif WITH_IOS_SIMULATOR
+#define SHADER_PLATFORM_METAL		SP_METAL_SIM
+#define SHADER_PLATFORM_METAL_MRT	SP_METAL_MRT
+#else
 #define SHADER_PLATFORM_METAL_MRT	SP_METAL_MRT
 #define SHADER_PLATFORM_METAL		SP_METAL
-
-#endif // PLATFORM_TVOS
-
+#endif
 
 //------------------------------------------------------------------------------
 // MARK: - iOS/tvOS Platform Dynamic RHI Routines
