[shaders] metal shader format: enable minifier and migrate to independent preprocessing API

diff --git a/Engine/Source/Developer/Apple/MetalShaderFormat/Private/MetalDerivedData.cpp b/Engine/Source/Developer/Apple/MetalShaderFormat/Private/MetalDerivedData.cpp
--- a/Engine/Source/Developer/Apple/MetalShaderFormat/Private/MetalDerivedData.cpp
+++ b/Engine/Source/Developer/Apple/MetalShaderFormat/Private/MetalDerivedData.cpp
@@ -191,10 +191,9 @@ static void Patch16bitInHlslSource(const FShaderCompilerInput& Input, std::strin
     }
 }
 
-bool DoCompileMetalShader(
+void DoCompileMetalShader(
 	const FShaderCompilerInput& Input,
 	FShaderCompilerOutput& Output,
-	const FString& WorkingDirectory,
 	const FString& InPreprocessedShader,
 	FSHAHash GUIDHash,
 	uint32 VersionEnum,
@@ -1411,7 +1410,6 @@ bool DoCompileMetalShader(
 	{
 		Output.Target = Input.Target;
 		BuildMetalShaderOutput(Output, Input, GUIDHash, MetalSource.c_str(), MetalSource.length(), CRCLen, CRC, VersionEnum, *Standard, *MinOSVersion, Output.Errors, OutputData.TypedBuffers, OutputData.InvariantBuffers, OutputData.TypedUAVs, OutputData.ConstantBuffers, bAllowFastIntrinsics);
-		return Output.bSucceeded;
 	}
 	else
 	{
@@ -1425,6 +1423,5 @@ bool DoCompileMetalShader(
 			}
 			GLog->Flush();
 		}
-		return false;
 	}
 }
diff --git a/Engine/Source/Developer/Apple/MetalShaderFormat/Private/MetalDerivedData.h b/Engine/Source/Developer/Apple/MetalShaderFormat/Private/MetalDerivedData.h
--- a/Engine/Source/Developer/Apple/MetalShaderFormat/Private/MetalDerivedData.h
+++ b/Engine/Source/Developer/Apple/MetalShaderFormat/Private/MetalDerivedData.h
@@ -90,10 +90,9 @@ struct FMetalShaderPreprocessed
 	}
 };
 
-bool DoCompileMetalShader(
+void DoCompileMetalShader(
 	const FShaderCompilerInput& Input,
 	FShaderCompilerOutput& Output,
-	const FString& WorkingDirectory,
 	const FString& PreprocessedShader,
 	FSHAHash GUIDHash,
 	uint32 VersionEnum,
diff --git a/Engine/Source/Developer/Apple/MetalShaderFormat/Private/MetalShaderCompiler.cpp b/Engine/Source/Developer/Apple/MetalShaderFormat/Private/MetalShaderCompiler.cpp
--- a/Engine/Source/Developer/Apple/MetalShaderFormat/Private/MetalShaderCompiler.cpp
+++ b/Engine/Source/Developer/Apple/MetalShaderFormat/Private/MetalShaderCompiler.cpp
@@ -12,6 +12,7 @@
 #include "ShaderCompilerDefinitions.h"
 #include "ShaderCore.h"
 #include "ShaderParameterParser.h"
+#include "ShaderPreprocessTypes.h"
 #include "Serialization/MemoryWriter.h"
 #include "Serialization/MemoryReader.h"
 
@@ -758,43 +759,36 @@ void BuildMetalShaderOutput(
 	External interface.
 ------------------------------------------------------------------------------*/
 
-void CompileShader_Metal(const FShaderCompilerInput& _Input,FShaderCompilerOutput& Output,const FString& WorkingDirectory)
+bool PreprocessMetalShader(const FShaderCompilerInput& Input, const FShaderCompilerEnvironment& Environment, FShaderPreprocessOutput& PreprocessOutput)
 {
+	const EShaderFrequency Frequency = (EShaderFrequency)Input.Target.Frequency;
+	if (!(Frequency == SF_Vertex || Frequency == SF_Pixel || Frequency == SF_Compute))
+	{
+		PreprocessOutput.LogError(FString::Printf(
+			TEXT("%s shaders not supported for use in Metal."),
+			CrossCompiler::GetFrequencyName(Frequency)
+			));
+		return false;
+	}
+
 	PRAGMA_DISABLE_DEPRECATION_WARNINGS		// FShaderCompilerDefinitions will be made internal in the future, marked deprecated until then
 
-	auto Input = _Input;
-	FString PreprocessedShader;
 	FShaderCompilerDefinitions AdditionalDefines;
 
 	// Work out which standard we need, this is dependent on the shader platform.
 	// TODO: Read from toolchain class
 	const bool bIsMobile = FMetalCompilerToolchain::Get()->IsMobile((EShaderPlatform) Input.Target.Platform);
-	TCHAR const* StandardPlatform = nullptr;
 	if (bIsMobile)
 	{
-		StandardPlatform = TEXT("ios");
 		AdditionalDefines.SetDefine(TEXT("IOS"), 1);
 	}
 	else
 	{
-		StandardPlatform = TEXT("macos");
 		AdditionalDefines.SetDefine(TEXT("MAC"), 1);
 	}
 	
 	AdditionalDefines.SetDefine(TEXT("COMPILER_METAL"), 1);
 	
-	EMetalGPUSemantics Semantics = EMetalGPUSemanticsMobile;
-	
-	uint32 VersionEnum = GMetalDefaultShadingLanguageVersion;
-	bool bFoundVersion = Input.Environment.GetCompileArgument(TEXT("SHADER_LANGUAGE_VERSION"), VersionEnum);
-	if (!bFoundVersion)
-	{
-		new(Output.Errors) FShaderCompilerError(*FString::Printf(TEXT("Missing SHADER_LANGUAGE_VERSION compile argument; Falling back to default value %d"), GMetalDefaultShadingLanguageVersion));
-	}
-
-	// TODO read from toolchain
-	bool bAppleTV = (Input.ShaderFormat == NAME_SF_METAL_TVOS || Input.ShaderFormat == NAME_SF_METAL_MRT_TVOS);
-	bool bIsSimulator = false;
 	if (Input.ShaderFormat == NAME_SF_METAL || Input.ShaderFormat == NAME_SF_METAL_TVOS)
 	{
 		AdditionalDefines.SetDefine(TEXT("METAL_PROFILE"), 1);
@@ -802,57 +796,144 @@ void CompileShader_Metal(const FShaderCompilerInput& _Input,FShaderCompilerOutpu
 	else if (Input.ShaderFormat == NAME_SF_METAL_SIM)
 	{
 		AdditionalDefines.SetDefine(TEXT("METAL_PROFILE"), 1);
-		bIsSimulator = true;
 	}
 	else if (Input.ShaderFormat == NAME_SF_METAL_MRT || Input.ShaderFormat == NAME_SF_METAL_MRT_TVOS)
 	{
 		AdditionalDefines.SetDefine(TEXT("METAL_MRT_PROFILE"), 1);
-		Semantics = EMetalGPUSemanticsTBDRDesktop;
 	}
 	else if (Input.ShaderFormat == NAME_SF_METAL_MACES3_1)
 	{
 		AdditionalDefines.SetDefine(TEXT("METAL_PROFILE"), 1);
-		Semantics = EMetalGPUSemanticsImmediateDesktop;
 	}
 	else if (Input.ShaderFormat == NAME_SF_METAL_SM5)
 	{
 		AdditionalDefines.SetDefine(TEXT("METAL_SM5_PROFILE"), 1);
 		AdditionalDefines.SetDefine(TEXT("USING_VERTEX_SHADER_LAYER"), 1);
-		Semantics = EMetalGPUSemanticsImmediateDesktop;
 	}
     else if (Input.ShaderFormat == NAME_SF_METAL_SM6)
     {
         AdditionalDefines.SetDefine(TEXT("METAL_SM6_PROFILE"), 1);
         AdditionalDefines.SetDefine(TEXT("USING_VERTEX_SHADER_LAYER"), 1);
-        Semantics = EMetalGPUSemanticsImmediateDesktop;
     }
 	else if (Input.ShaderFormat == NAME_SF_METAL_MRT_MAC)
 	{
 		AdditionalDefines.SetDefine(TEXT("METAL_MRT_PROFILE"), 1);
-		Semantics = EMetalGPUSemanticsTBDRDesktop;
 	}
 	else
 	{
-		Output.bSucceeded = false;
-		new(Output.Errors) FShaderCompilerError(*FString::Printf(TEXT("Invalid shader format '%s' passed to compiler."), *Input.ShaderFormat.ToString()));
-		return;
+		PreprocessOutput.LogError(FString::Printf(TEXT("Invalid shader format '%s' passed to compiler."), *Input.ShaderFormat.ToString()));
+		return false;
 	}
 	
 
 	AdditionalDefines.SetDefine(TEXT("COMPILER_HLSLCC"), 2);
-	
+
+	if (Input.Environment.FullPrecisionInPS)
+	{
+		AdditionalDefines.SetDefine(TEXT("FORCE_FLOATS"), (uint32)1);
+	}
+
+	if (Environment.CompilerFlags.Contains(CFLAG_AvoidFlowControl) 
+		|| Environment.CompilerFlags.Contains(CFLAG_PreferFlowControl))
+	{
+		AdditionalDefines.SetDefine(TEXT("COMPILER_SUPPORTS_ATTRIBUTES"), (uint32)0);
+	}
+	else
+	{
+		AdditionalDefines.SetDefine(TEXT("COMPILER_SUPPORTS_ATTRIBUTES"), (uint32)1);
+	}
+
+	bool bUsesInlineRayTracing = Environment.CompilerFlags.Contains(CFLAG_InlineRayTracing);
+	if (bUsesInlineRayTracing)
+	{
+		AdditionalDefines.SetDefine(TEXT("PLATFORM_SUPPORTS_INLINE_RAY_TRACING"), 1);
+	}
+
+	AdditionalDefines.SetDefine(TEXT("COMPILER_SUPPORTS_DUAL_SOURCE_BLENDING_SLOT_DECORATION"), (uint32)1);
+
+	PRAGMA_ENABLE_DEPRECATION_WARNINGS
+
+	if (Input.bSkipPreprocessedCache)
+	{
+		if (!FFileHelper::LoadFileToString(PreprocessOutput.EditSource(), *Input.VirtualSourceFilePath))
+		{
+			return false;
+		}
+
+		// Remove const as we are on debug-only mode
+		CrossCompiler::CreateEnvironmentFromResourceTable(PreprocessOutput.EditSource(), (FShaderCompilerEnvironment&)Input.Environment);
+	}
+	else
+	{
+		if (!PreprocessShader(PreprocessOutput, Input, Environment, AdditionalDefines))
+		{
+			// The preprocessing stage will add any relevant errors.
+			return false;
+		}
+	}
+
+	if (!PreprocessOutput.ParseAndModify(Input, Environment, nullptr))
+	{
+		// The FShaderParameterParser will add any relevant errors.
+		return false;
+	}
+
+	RemoveUniformBuffersFromSource(Input.Environment, PreprocessOutput.EditSource());
+
+	// Process TEXT macro.
+	TransformStringIntoCharacterArray(PreprocessOutput.EditSource());
+
+	// Run the shader minifier
+	#if UE_METAL_SHADER_COMPILER_ALLOW_DEAD_CODE_REMOVAL
+	if (Input.Environment.CompilerFlags.Contains(CFLAG_RemoveDeadCode))
+	{
+		UE::ShaderCompilerCommon::RemoveDeadCode(PreprocessOutput.EditSource(), Input.EntryPointName, PreprocessOutput.EditErrors());
+	}
+	#endif // UE_METAL_SHADER_COMPILER_ALLOW_DEAD_CODE_REMOVAL
+
+	return true;
+}
+
+void CompileMetalShader(const FShaderCompilerInput& Input, const FShaderPreprocessOutput& PreprocessOutput, FShaderCompilerOutput& Output)
+{
+	EMetalGPUSemantics Semantics = EMetalGPUSemanticsMobile;
+	if (Input.ShaderFormat == NAME_SF_METAL_MRT
+		|| Input.ShaderFormat == NAME_SF_METAL_MRT_TVOS
+		|| Input.ShaderFormat == NAME_SF_METAL_MRT_MAC)
+	{
+		Semantics = EMetalGPUSemanticsTBDRDesktop;
+	}
+	else if (Input.ShaderFormat == NAME_SF_METAL_MACES3_1
+		|| Input.ShaderFormat == NAME_SF_METAL_SM5
+		|| Input.ShaderFormat == NAME_SF_METAL_SM6)
+	{
+		Semantics = EMetalGPUSemanticsImmediateDesktop;
+	}
+
+	uint32 VersionEnum = GMetalDefaultShadingLanguageVersion;
+	bool bFoundVersion = Input.Environment.GetCompileArgument(TEXT("SHADER_LANGUAGE_VERSION"), VersionEnum);
+	if (!bFoundVersion)
+	{
+		new(Output.Errors) FShaderCompilerError(*FString::Printf(TEXT("Missing SHADER_LANGUAGE_VERSION compile argument; Falling back to default value %d"), GMetalDefaultShadingLanguageVersion));
+	}
+
+	// TODO read from toolchain
+	const bool bIsMobile = FMetalCompilerToolchain::Get()->IsMobile((EShaderPlatform)Input.Target.Platform);
+	const bool bAppleTV = (Input.ShaderFormat == NAME_SF_METAL_TVOS || Input.ShaderFormat == NAME_SF_METAL_MRT_TVOS);
+	const bool bIsSimulator = (Input.ShaderFormat == NAME_SF_METAL_SIM);
+
 	FString MinOSVersion;
 	FString StandardVersion;
-	switch(VersionEnum)
+	switch (VersionEnum)
 	{
-    case 9:
-        StandardVersion = TEXT("3.1");
-        if (bAppleTV)
-        {
-            MinOSVersion = TEXT("-mtvos-version-min=17.0");
-        }
-        else if (bIsMobile)
-        {
+	case 9:
+		StandardVersion = TEXT("3.1");
+		if (bAppleTV)
+		{
+			MinOSVersion = TEXT("-mtvos-version-min=17.0");
+		}
+		else if (bIsMobile)
+		{
 			if (bIsSimulator)
 			{
 				MinOSVersion = TEXT("-miphonesimulator-version-min=17.0");
@@ -861,20 +942,20 @@ void CompileShader_Metal(const FShaderCompilerInput& _Input,FShaderCompilerOutpu
 			{
 				MinOSVersion = TEXT("-mios-version-min=17.0");
 			}
-        }
-        else
-        {
-            MinOSVersion = TEXT("-mmacosx-version-min=14");
-        }
-            break;
-    case 8:
-        StandardVersion = TEXT("3.0");
-        if (bAppleTV)
-        {
-            MinOSVersion = TEXT("-mtvos-version-min=16.0");
-        }
-        else if (bIsMobile)
-        {
+		}
+		else
+		{
+			MinOSVersion = TEXT("-mmacosx-version-min=14");
+		}
+		break;
+	case 8:
+		StandardVersion = TEXT("3.0");
+		if (bAppleTV)
+		{
+			MinOSVersion = TEXT("-mtvos-version-min=16.0");
+		}
+		else if (bIsMobile)
+		{
 			if (bIsSimulator)
 			{
 				MinOSVersion = TEXT("-miphonesimulator-version-min=16.0");
@@ -883,12 +964,12 @@ void CompileShader_Metal(const FShaderCompilerInput& _Input,FShaderCompilerOutpu
 			{
 				MinOSVersion = TEXT("-mios-version-min=16.0");
 			}
-        }
-        else
-        {
-            MinOSVersion = TEXT("-mmacosx-version-min=13");
-        }
-        break;
+		}
+		else
+		{
+			MinOSVersion = TEXT("-mmacosx-version-min=13");
+		}
+		break;
 
 	case 7:
 		StandardVersion = TEXT("2.4");
@@ -970,128 +1051,44 @@ void CompileShader_Metal(const FShaderCompilerInput& _Input,FShaderCompilerOutpu
 	default:
 		Output.bSucceeded = false;
 		{
-            FString EngineIdentifier = FEngineVersion::Current().ToString(EVersionComponent::Minor);
+			FString EngineIdentifier = FEngineVersion::Current().ToString(EVersionComponent::Minor);
 			FShaderCompilerError* NewError = new(Output.Errors) FShaderCompilerError();
-            NewError->StrippedErrorMessage = FString::Printf(TEXT("Minimum Metal Version is 2.4 in UE %s"), *EngineIdentifier);
+			NewError->StrippedErrorMessage = FString::Printf(TEXT("Minimum Metal Version is 2.4 in UE %s"), *EngineIdentifier);
 			return;
 		}
 		break;
 	}
-	
-    if (Input.Environment.FullPrecisionInPS)
-    {
-        AdditionalDefines.SetDefine(TEXT("FORCE_FLOATS"), (uint32)1);
-    }
 
-    FString Standard;
-    if (VersionEnum >= 8)
-    {
-        Standard = FString::Printf(TEXT("-std=metal%s"), *StandardVersion);
-    }
-    else
-    {
-        Standard = FString::Printf(TEXT("-std=%s-metal%s"), StandardPlatform, *StandardVersion);
-    }
-	
+	TCHAR const* StandardPlatform = bIsMobile ? TEXT("ios") : TEXT("macos");
+	FString Standard;
+	if (VersionEnum >= 8)
+	{
+		Standard = FString::Printf(TEXT("-std=metal%s"), *StandardVersion);
+	}
+	else
+	{
+		Standard = FString::Printf(TEXT("-std=%s-metal%s"), StandardPlatform, *StandardVersion);
+	}
+
 	bool const bDirectCompile = FParse::Param(FCommandLine::Get(), TEXT("directcompile"));
 	if (bDirectCompile)
 	{
-		Input.DumpDebugInfoPath = FPaths::GetPath(Input.VirtualSourceFilePath);
+		const_cast<FShaderCompilerInput&>(Input).DumpDebugInfoPath = FPaths::GetPath(Input.VirtualSourceFilePath);
 	}
-	
+
 	const bool bDumpDebugInfo = Input.DumpDebugInfoEnabled();
 
 	// Allow the shader pipeline to override the platform default in here.
 	uint32 MaxUnrollLoops = 32;
 	if (Input.Environment.CompilerFlags.Contains(CFLAG_AvoidFlowControl))
 	{
-		AdditionalDefines.SetDefine(TEXT("COMPILER_SUPPORTS_ATTRIBUTES"), (uint32)0);
 		MaxUnrollLoops = 1024; // Max. permitted by hlslcc
 	}
 	else if (Input.Environment.CompilerFlags.Contains(CFLAG_PreferFlowControl))
 	{
-		AdditionalDefines.SetDefine(TEXT("COMPILER_SUPPORTS_ATTRIBUTES"), (uint32)0);
 		MaxUnrollLoops = 0;
 	}
-	else
-	{
-		AdditionalDefines.SetDefine(TEXT("COMPILER_SUPPORTS_ATTRIBUTES"), (uint32)1);
-	}
-
-	bool bUsesInlineRayTracing = Input.Environment.CompilerFlags.Contains(CFLAG_InlineRayTracing);
-	if (bUsesInlineRayTracing)
-	{
-		AdditionalDefines.SetDefine(TEXT("PLATFORM_SUPPORTS_INLINE_RAY_TRACING"), 1);
-	}
 
-	AdditionalDefines.SetDefine(TEXT("COMPILER_SUPPORTS_DUAL_SOURCE_BLENDING_SLOT_DECORATION"), (uint32)1);
-
-	PRAGMA_ENABLE_DEPRECATION_WARNINGS
-
-	const double StartPreprocessTime = FPlatformTime::Seconds();
-
-	if (Input.bSkipPreprocessedCache)
-	{
-		if (!FFileHelper::LoadFileToString(PreprocessedShader, *Input.VirtualSourceFilePath))
-		{
-			return;
-		}
-
-		// Remove const as we are on debug-only mode
-		CrossCompiler::CreateEnvironmentFromResourceTable(PreprocessedShader, (FShaderCompilerEnvironment&)Input.Environment);
-	}
-	else
-	{
-		if (!PreprocessShader(PreprocessedShader, Output, Input, AdditionalDefines))
-		{
-			// The preprocessing stage will add any relevant errors.
-			return;
-		}
-	}
-
-	char* MetalShaderSource = NULL;
-	char* ErrorLog = NULL;
-
-	const EShaderFrequency Frequency = (EShaderFrequency)Input.Target.Frequency;
-	if (!(Frequency == SF_Vertex || Frequency == SF_Pixel || Frequency == SF_Compute))
-	{
-		Output.bSucceeded = false;
-		FShaderCompilerError* NewError = new(Output.Errors) FShaderCompilerError();
-		NewError->StrippedErrorMessage = FString::Printf(
-			TEXT("%s shaders not supported for use in Metal."),
-			CrossCompiler::GetFrequencyName(Frequency)
-			);
-		return;
-	}
-
-	FShaderParameterParser ShaderParameterParser(Input.Environment.CompilerFlags);
-	if (!ShaderParameterParser.ParseAndModify(Input, Output.Errors, PreprocessedShader))
-	{
-		// The FShaderParameterParser will add any relevant errors.
-		return;
-	}
-
-	// This requires removing the HLSLCC_NoPreprocess flag later on!
-	RemoveUniformBuffersFromSource(Input.Environment, PreprocessedShader);
-
-	// Process TEXT macro.
-	TransformStringIntoCharacterArray(PreprocessedShader);
-
-	// Run the shader minifier
-	#if UE_METAL_SHADER_COMPILER_ALLOW_DEAD_CODE_REMOVAL
-	if (Input.Environment.CompilerFlags.Contains(CFLAG_RemoveDeadCode))
-	{
-		UE::ShaderCompilerCommon::RemoveDeadCode(PreprocessedShader, Input.EntryPointName, Output.Errors);
-	}
-	#endif // UE_METAL_SHADER_COMPILER_ALLOW_DEAD_CODE_REMOVAL
-
-	Output.PreprocessTime = FPlatformTime::Seconds() - StartPreprocessTime;
-	
-	// Write out the preprocessed file and a batch file to compile it if requested (DumpDebugInfoPath is valid)
-	if (bDumpDebugInfo && !bDirectCompile)
-	{
-		UE::ShaderCompilerCommon::DumpDebugShaderData(Input, PreprocessedShader);
-	}
 
 	FSHAHash GUIDHash;
 	if (!bDirectCompile)
@@ -1107,13 +1104,13 @@ void CompileShader_Metal(const FShaderCompilerInput& _Input,FShaderCompilerOutpu
 		FSHA1::HashBuffer(&Guid, sizeof(FGuid), GUIDHash.Hash);
 	}
 
-	bool bCompiled = DoCompileMetalShader(Input, Output, WorkingDirectory, PreprocessedShader, GUIDHash, VersionEnum, Semantics, MaxUnrollLoops, Frequency, bDumpDebugInfo, Standard, MinOSVersion);
-	if (bCompiled && Output.bSucceeded)
-	{
-
-	}
+	DoCompileMetalShader(Input, Output, PreprocessOutput.GetSource(), GUIDHash, VersionEnum, Semantics, MaxUnrollLoops, (EShaderFrequency)Input.Target.Frequency, bDumpDebugInfo, Standard, MinOSVersion);
+	PreprocessOutput.GetParameterParser().ValidateShaderParameterTypes(Input, bIsMobile, Output);
+}
 
-	ShaderParameterParser.ValidateShaderParameterTypes(Input, bIsMobile, Output);
+void OutputMetalDebugData(const FShaderCompilerInput& Input, const FShaderPreprocessOutput& PreprocessOutput, const FShaderCompilerOutput& Output)
+{
+	UE::ShaderCompilerCommon::DumpExtendedDebugShaderData(Input, PreprocessOutput, Output);
 }
 
 bool StripShader_Metal(TArray<uint8>& Code, class FString const& DebugPath, bool const bNative)
diff --git a/Engine/Source/Developer/Apple/MetalShaderFormat/Private/MetalShaderFormat.cpp b/Engine/Source/Developer/Apple/MetalShaderFormat/Private/MetalShaderFormat.cpp
--- a/Engine/Source/Developer/Apple/MetalShaderFormat/Private/MetalShaderFormat.cpp
+++ b/Engine/Source/Developer/Apple/MetalShaderFormat/Private/MetalShaderFormat.cpp
@@ -27,6 +27,10 @@ DEFINE_LOG_CATEGORY(LogMetalShaderCompiler)
 // Set this define to get additional logging information about Metal toolchain setup.
 #define CHECK_METAL_COMPILER_TOOLCHAIN_SETUP 0
 
+extern bool PreprocessMetalShader(const FShaderCompilerInput& Input, const FShaderCompilerEnvironment& Environment, FShaderPreprocessOutput& PreprocessOutput);
+extern void CompileMetalShader(const FShaderCompilerInput& Input, const FShaderPreprocessOutput& PreprocessOutput, FShaderCompilerOutput& Output);
+extern void OutputMetalDebugData(const FShaderCompilerInput& Input, const FShaderPreprocessOutput& PreprocessOutput, const FShaderCompilerOutput& Output);
+
 extern bool StripShader_Metal(TArray<uint8>& Code, class FString const& DebugPath, bool const bNative);
 extern uint64 AppendShader_Metal(class FString const& ArchivePath, const FSHAHash& Hash, TArray<uint8>& Code);
 extern bool FinalizeLibrary_Metal(class FName const& Format, class FString const& ArchivePath, class FString const& LibraryPath, TSet<uint64> const& Shaders, class FString const& DebugOutputDir);
@@ -55,8 +59,64 @@ public:
 	}
 	virtual uint32 GetVersion(FName Format) const override final
 	{
-		return GetMetalFormatVersion(Format);
+		static_assert(sizeof(FMetalShaderFormat::FVersion) == sizeof(uint32), "Out of bits!");
+		union
+		{
+			FMetalShaderFormat::FVersion Version;
+			uint32 Raw;
+		} Version;
+
+		// If there's no compiler on this machine, this is irrelevant so just return 0
+		if (!FMetalCompilerToolchain::Get()->IsCompilerAvailable())
+		{
+			return 0;
+		}
+
+		bool bUseFullMetalVersion = false;
+		EShaderPlatform ShaderPlatform = FMetalCompilerToolchain::MetalShaderFormatToLegacyShaderPlatform(Format);
+
+		if (FMetalCompilerToolchain::IsMobile(ShaderPlatform))
+		{
+			GConfig->GetBool(TEXT("/Script/IOSRuntimeSettings.IOSRuntimeSettings"), TEXT("UseFullMetalVersionInShaderVersion"), bUseFullMetalVersion, GEngineIni);
+		}
+		else
+		{
+			GConfig->GetBool(TEXT("/Script/MacTargetPlatform.MacTargetSettings"), TEXT("UseFullMetalVersionInShaderVersion"), bUseFullMetalVersion, GEngineIni);
+		}
+
+		FMetalCompilerToolchain::PackedVersion MetalVersionNumber = FMetalCompilerToolchain::Get()->GetCompilerVersion(ShaderPlatform);
+		uint16 HashValue = MetalVersionNumber.Major;
+
+		if (bUseFullMetalVersion)
+		{
+			// Use entire Metal version if .ini settings instruct us to do so (e.g. p4 dev build)
+			HashValue ^= MetalVersionNumber.Minor;
+			HashValue ^= MetalVersionNumber.Patch;
+		}
+		else
+		{
+			// Only use Metal major version (e.g. Installed build)
+			// Since Metal minor/patch version changes every Xcode minor version, we don't want users to rebuild shaders for every minor version update
+		}
+
+		Version.Version.XcodeVersion = HashValue;
+		Version.Version.Format = FMetalShaderFormat::HEADER_VERSION;
+		Version.Version.HLSLCCMinor = HLSLCC_VersionMinor;
+
+		// Check that we didn't overwrite any bits
+		check(Version.Version.XcodeVersion == HashValue);
+		check(Version.Version.Format == FMetalShaderFormat::HEADER_VERSION);
+		check(Version.Version.HLSLCCMinor == HLSLCC_VersionMinor);
+
+		uint32 Result = Version.Raw;
+
+#if UE_METAL_SHADER_COMPILER_ALLOW_DEAD_CODE_REMOVAL
+		Result = HashCombine(Result, 0x75E2FE85);
+#endif // UE_METAL_SHADER_COMPILER_ALLOW_DEAD_CODE_REMOVAL
+
+		return Result;
 	}
+
 	virtual void GetSupportedFormats(TArray<FName>& OutFormats) const override final
 	{
 		OutFormats.Add(NAME_SF_METAL);
@@ -64,29 +124,73 @@ public:
 		OutFormats.Add(NAME_SF_METAL_TVOS);
 		OutFormats.Add(NAME_SF_METAL_MRT_TVOS);
 		OutFormats.Add(NAME_SF_METAL_SM5);
-        OutFormats.Add(NAME_SF_METAL_SM6);
-        OutFormats.Add(NAME_SF_METAL_SIM);
+		OutFormats.Add(NAME_SF_METAL_SM6);
+		OutFormats.Add(NAME_SF_METAL_SIM);
 		OutFormats.Add(NAME_SF_METAL_MACES3_1);
 		OutFormats.Add(NAME_SF_METAL_MRT_MAC);
 	}
-	virtual void CompileShader(FName Format, const struct FShaderCompilerInput& Input, struct FShaderCompilerOutput& Output,const FString& WorkingDirectory) const override final
+
+	virtual bool SupportsIndependentPreprocessing() const override final
+	{
+		return true;
+	}
+
+	void CheckShaderFormat(FName Format) const
+	{
+		check(Format == NAME_SF_METAL
+			|| Format == NAME_SF_METAL_MRT
+			|| Format == NAME_SF_METAL_TVOS
+			|| Format == NAME_SF_METAL_MRT_TVOS
+			|| Format == NAME_SF_METAL_SM5
+			|| Format == NAME_SF_METAL_SM6
+			|| Format == NAME_SF_METAL_SIM
+			|| Format == NAME_SF_METAL_MACES3_1
+			|| Format == NAME_SF_METAL_MRT_MAC);
+	}
+
+	virtual bool PreprocessShader(
+		const FShaderCompilerInput& Input,
+		const FShaderCompilerEnvironment& Environment,
+		FShaderPreprocessOutput& PreprocessOutput) const override final
 	{
-		check(Format == NAME_SF_METAL || Format == NAME_SF_METAL_MRT || Format == NAME_SF_METAL_TVOS || Format == NAME_SF_METAL_MRT_TVOS || Format == NAME_SF_METAL_SM5 || Format == NAME_SF_METAL_SM6 || Format == NAME_SF_METAL_SIM || Format == NAME_SF_METAL_MACES3_1 || Format == NAME_SF_METAL_MRT_MAC);
-		CompileShader_Metal(Input, Output, WorkingDirectory);
+		CheckShaderFormat(Input.ShaderFormat);
+		return PreprocessMetalShader(Input, Environment, PreprocessOutput);
 	}
+
+	virtual void CompilePreprocessedShader(
+		const FShaderCompilerInput& Input,
+		const FShaderPreprocessOutput& PreprocessOutput,
+		FShaderCompilerOutput& Output,
+		const FString& WorkingDirectory) const override final
+	{
+		CheckShaderFormat(Input.ShaderFormat);
+		CompileMetalShader(Input, PreprocessOutput, Output);
+	}
+
+	virtual void OutputDebugData(
+		const FShaderCompilerInput& Input,
+		const FShaderPreprocessOutput& PreprocessOutput,
+		const FShaderCompilerOutput& Output) const override final
+	{
+		OutputMetalDebugData(Input, PreprocessOutput, Output);
+	}
+
 	virtual bool CanStripShaderCode(bool const bNativeFormat) const override final
 	{
 		return CanCompileBinaryShaders() && bNativeFormat;
 	}
+
 	virtual bool StripShaderCode( TArray<uint8>& Code, FString const& DebugOutputDir, bool const bNative ) const override final
 	{
 		return StripShader_Metal(Code, DebugOutputDir, bNative);
     }
+
 	virtual bool SupportsShaderArchives() const override 
 	{ 
 		return CanCompileBinaryShaders();
 	}
-    virtual bool CreateShaderArchive(FString const& LibraryName,
+
+	virtual bool CreateShaderArchive(FString const& LibraryName,
 		FName ShaderFormatAndShaderPlatformName,
 		const FString& WorkingDirectory,
 		const FString& OutputDir,
@@ -94,7 +198,7 @@ public:
 		const FSerializedShaderArchive& InSerializedShaders,
 		const TArray<TArray<uint8>>& ShaderCode,
 		TArray<FString>* OutputFiles) const override final
-    {
+	{
 		const int32 NumShadersPerLibrary = 10000;
 		check(LibraryName.Len() > 0);
 
@@ -104,7 +208,7 @@ public:
 		check(Components.Num() == 2);
 		FName ShaderFormatName(Components[0]);
 
-		check(ShaderFormatName == NAME_SF_METAL || ShaderFormatName == NAME_SF_METAL_MRT || ShaderFormatName == NAME_SF_METAL_TVOS || ShaderFormatName == NAME_SF_METAL_MRT_TVOS || ShaderFormatName == NAME_SF_METAL_SM5 || ShaderFormatName == NAME_SF_METAL_SM6  || ShaderFormatName == NAME_SF_METAL_SIM || ShaderFormatName == NAME_SF_METAL_MACES3_1 || ShaderFormatName == NAME_SF_METAL_MRT_MAC);
+		check(ShaderFormatName == NAME_SF_METAL || ShaderFormatName == NAME_SF_METAL_MRT || ShaderFormatName == NAME_SF_METAL_TVOS || ShaderFormatName == NAME_SF_METAL_MRT_TVOS || ShaderFormatName == NAME_SF_METAL_SM5 || ShaderFormatName == NAME_SF_METAL_SM6 || ShaderFormatName == NAME_SF_METAL_SIM || ShaderFormatName == NAME_SF_METAL_MACES3_1 || ShaderFormatName == NAME_SF_METAL_MRT_MAC);
 
 		const FString ArchivePath = (WorkingDirectory / ShaderFormatAndShaderPlatformName.GetPlainNameString());
 		IFileManager::Get().DeleteDirectory(*ArchivePath, false, true);
@@ -162,61 +266,61 @@ public:
 
 			// Enqueue the library compilation as a task so we can go wide
 			FGraphEventRef CompletionFence = FFunctionGraphTask::CreateAndDispatchWhenReady([ShaderFormatName, ArchivePath, LibraryPath, PartialShaders, DebugOutputDir, &CompiledLibraries]()
-			{
-				if (FinalizeLibrary_Metal(ShaderFormatName, ArchivePath, LibraryPath, PartialShaders, DebugOutputDir))
 				{
-					FPlatformAtomics::InterlockedIncrement(&CompiledLibraries);
-				}
-			}, TStatId(), NULL, ENamedThreads::AnyThread);
+					if (FinalizeLibrary_Metal(ShaderFormatName, ArchivePath, LibraryPath, PartialShaders, DebugOutputDir))
+					{
+						FPlatformAtomics::InterlockedIncrement(&CompiledLibraries);
+					}
+				}, TStatId(), NULL, ENamedThreads::AnyThread);
 
 			Tasks.Add(CompletionFence);
 		}
 
 #if WITH_ENGINE
 		FGraphEventRef DebugDataCompletionFence = FFunctionGraphTask::CreateAndDispatchWhenReady([ShaderFormatAndShaderPlatformName, OutputDir, LibraryPlatformName, DebugOutputDir]()
-		{
-			//TODO add a check in here - this will only work if we have shader archiving with debug info set.
+			{
+				//TODO add a check in here - this will only work if we have shader archiving with debug info set.
 
-			//We want to archive all the metal shader source files so that they can be unarchived into a debug location
-			//This allows the debugging of optimised metal shaders within the xcode tool set
-			//Currently using the 'tar' system tool to create a compressed tape archive
+				//We want to archive all the metal shader source files so that they can be unarchived into a debug location
+				//This allows the debugging of optimised metal shaders within the xcode tool set
+				//Currently using the 'tar' system tool to create a compressed tape archive
 
-			//Place the archive in the same position as the .metallib file
-			FString CompressedDir = (OutputDir / TEXT("../MetaData/ShaderDebug/"));
-			IFileManager::Get().MakeDirectory(*CompressedDir, true);
+				//Place the archive in the same position as the .metallib file
+				FString CompressedDir = (OutputDir / TEXT("../MetaData/ShaderDebug/"));
+				IFileManager::Get().MakeDirectory(*CompressedDir, true);
 
-			FString CompressedPath = (CompressedDir / LibraryPlatformName) + TEXT(".zip");
+				FString CompressedPath = (CompressedDir / LibraryPlatformName) + TEXT(".zip");
 
-			IPlatformFile& PlatformFile = FPlatformFileManager::Get().GetPlatformFile();
-			IFileHandle* ZipFile = PlatformFile.OpenWrite(*CompressedPath);
-			if (ZipFile)
-			{
-				FZipArchiveWriter* ZipWriter = new FZipArchiveWriter(ZipFile);
+				IPlatformFile& PlatformFile = FPlatformFileManager::Get().GetPlatformFile();
+				IFileHandle* ZipFile = PlatformFile.OpenWrite(*CompressedPath);
+				if (ZipFile)
+				{
+					FZipArchiveWriter* ZipWriter = new FZipArchiveWriter(ZipFile);
 
-				//Find the metal source files
-				TArray<FString> FilesToArchive;
-				IFileManager::Get().FindFilesRecursive(FilesToArchive, *DebugOutputDir, TEXT("*.metal"), true, false, false);
+					//Find the metal source files
+					TArray<FString> FilesToArchive;
+					IFileManager::Get().FindFilesRecursive(FilesToArchive, *DebugOutputDir, TEXT("*.metal"), true, false, false);
 
-				//Write the local file names into the target file
-				const FString DebugDir = DebugOutputDir / *ShaderFormatAndShaderPlatformName.GetPlainNameString();
+					//Write the local file names into the target file
+					const FString DebugDir = DebugOutputDir / *ShaderFormatAndShaderPlatformName.GetPlainNameString();
 
-				for (FString FileName : FilesToArchive)
-				{
-					TArray<uint8> FileData;
-					FFileHelper::LoadFileToArray(FileData, *FileName);
-					FPaths::MakePathRelativeTo(FileName, *DebugDir);
+					for (FString FileName : FilesToArchive)
+					{
+						TArray<uint8> FileData;
+						FFileHelper::LoadFileToArray(FileData, *FileName);
+						FPaths::MakePathRelativeTo(FileName, *DebugDir);
 
-					ZipWriter->AddFile(FileName, FileData, FDateTime::Now());
-				}
+						ZipWriter->AddFile(FileName, FileData, FDateTime::Now());
+					}
 
-				delete ZipWriter;
-				ZipWriter = nullptr;
-			}
-			else
-			{
-				UE_LOG(LogShaders, Error, TEXT("Failed to create Metal debug .zip output file \"%s\". Debug .zip export will be disabled."), *CompressedPath);
-			}
-		}, TStatId(), NULL, ENamedThreads::AnyThread);
+					delete ZipWriter;
+					ZipWriter = nullptr;
+				}
+				else
+				{
+					UE_LOG(LogShaders, Error, TEXT("Failed to create Metal debug .zip output file \"%s\". Debug .zip export will be disabled."), *CompressedPath);
+				}
+			}, TStatId(), NULL, ENamedThreads::AnyThread);
 		Tasks.Add(DebugDataCompletionFence);
 #endif // WITH_ENGINE
 
@@ -257,7 +361,7 @@ public:
 		return bOK;
 
 		//Map.Format = Format.GetPlainNameString();
-    }
+	}
 
 	virtual bool CanCompileBinaryShaders() const override final
 	{
@@ -273,66 +377,6 @@ public:
 	}
 };
 
-uint32 GetMetalFormatVersion(FName Format)
-{
-	static_assert(sizeof(FMetalShaderFormat::FVersion) == sizeof(uint32), "Out of bits!");
-	union
-	{
-		FMetalShaderFormat::FVersion Version;
-		uint32 Raw;
-	} Version;
-
-	// If there's no compiler on this machine, this is irrelevant so just return 0
-	if (!FMetalCompilerToolchain::Get()->IsCompilerAvailable())
-	{
-		return 0;
-	}
-	
-	bool bUseFullMetalVersion = false;
-	EShaderPlatform ShaderPlatform = FMetalCompilerToolchain::MetalShaderFormatToLegacyShaderPlatform(Format);
-
-	if(FMetalCompilerToolchain::IsMobile(ShaderPlatform))
-	{
-		GConfig->GetBool(TEXT("/Script/IOSRuntimeSettings.IOSRuntimeSettings"), TEXT("UseFullMetalVersionInShaderVersion"), bUseFullMetalVersion, GEngineIni);
-	}
-	else
-	{
-		GConfig->GetBool(TEXT("/Script/MacTargetPlatform.MacTargetSettings"), TEXT("UseFullMetalVersionInShaderVersion"), bUseFullMetalVersion, GEngineIni);
-	}
-	
-	FMetalCompilerToolchain::PackedVersion MetalVersionNumber = FMetalCompilerToolchain::Get()->GetCompilerVersion(ShaderPlatform);
-	uint16 HashValue = MetalVersionNumber.Major;
-	
-	if (bUseFullMetalVersion)
-	{
-		// Use entire Metal version if .ini settings instruct us to do so (e.g. p4 dev build)
-		HashValue ^= MetalVersionNumber.Minor;
-		HashValue ^= MetalVersionNumber.Patch;
-	}
-	else
-	{
-		// Only use Metal major version (e.g. Installed build)
-		// Since Metal minor/patch version changes every Xcode minor version, we don't want users to rebuild shaders for every minor version update
-	}
-
-	Version.Version.XcodeVersion = HashValue;
-	Version.Version.Format = FMetalShaderFormat::HEADER_VERSION;
-	Version.Version.HLSLCCMinor = HLSLCC_VersionMinor;
-	
-	// Check that we didn't overwrite any bits
-	check(Version.Version.XcodeVersion == HashValue);
-	check(Version.Version.Format == FMetalShaderFormat::HEADER_VERSION);
-	check(Version.Version.HLSLCCMinor == HLSLCC_VersionMinor);
-	
-	uint32 Result = Version.Raw;
-
-#if UE_METAL_SHADER_COMPILER_ALLOW_DEAD_CODE_REMOVAL
-	Result = HashCombine(Result, 0x75E2FE85);
-#endif // UE_METAL_SHADER_COMPILER_ALLOW_DEAD_CODE_REMOVAL
-
-	return Result;
-}
-
 /**
  * Module for Metal shaders
  */
@@ -365,7 +409,7 @@ public:
 	}
 };
 
-IMPLEMENT_MODULE( FMetalShaderFormatModule, MetalShaderFormat);
+IMPLEMENT_MODULE(FMetalShaderFormatModule, MetalShaderFormat);
 
 static FMetalCompilerToolchain::EMetalToolchainStatus ParseCompilerVersionAndTarget(const FString& OutputOfMetalDashV, FString& VersionString, FMetalCompilerToolchain::PackedVersion& PackedVersionNumber, FMetalCompilerToolchain::PackedVersion& PackedTargetNumber)
 {
diff --git a/Engine/Source/Developer/Apple/MetalShaderFormat/Private/MetalShaderFormat.h b/Engine/Source/Developer/Apple/MetalShaderFormat/Private/MetalShaderFormat.h
--- a/Engine/Source/Developer/Apple/MetalShaderFormat/Private/MetalShaderFormat.h
+++ b/Engine/Source/Developer/Apple/MetalShaderFormat/Private/MetalShaderFormat.h
@@ -13,7 +13,7 @@
 
 // Controls whether r.Shaders.RemoveDeadCode should be honored
 #ifndef UE_METAL_SHADER_COMPILER_ALLOW_DEAD_CODE_REMOVAL
-#define UE_METAL_SHADER_COMPILER_ALLOW_DEAD_CODE_REMOVAL 0 // Disabled initially, requires testing
+#define UE_METAL_SHADER_COMPILER_ALLOW_DEAD_CODE_REMOVAL 1
 #endif // UE_METAL_SHADER_COMPILER_ALLOW_DEAD_CODE_REMOVAL
 
 // IOS and TVOS use the mobile toolchain.
@@ -24,9 +24,6 @@ enum EAppleSDKType
 	AppleSDKCount,
 };
 
-extern void CompileShader_Metal(const struct FShaderCompilerInput& Input, struct FShaderCompilerOutput& Output, const class FString& WorkingDirectory);
-extern uint32 GetMetalFormatVersion(FName Format);
-
 static FName NAME_SF_METAL(TEXT("SF_METAL"));
 static FName NAME_SF_METAL_MRT(TEXT("SF_METAL_MRT"));
 static FName NAME_SF_METAL_TVOS(TEXT("SF_METAL_TVOS"));
