Shader versioning cleanup
- Changing all shader versions to Guid.
- Fixing shader version guids in shaders to not be in comments.
- Adding a shader version to RayTracingCommon.ush.

diff --git a/Engine/Shaders/Private/RayTracing/RayTracingCommon.ush b/Engine/Shaders/Private/RayTracing/RayTracingCommon.ush
--- a/Engine/Shaders/Private/RayTracing/RayTracingCommon.ush
+++ b/Engine/Shaders/Private/RayTracing/RayTracingCommon.ush
@@ -12,6 +12,9 @@
 // Make sure we recompile ray tracing shaders if the main shader version changes
 #include "/Engine/Public/ShaderVersion.ush"
 
+// Update this GUID to invalidate RayTracing shaders and force a shader recompilation.
+#pragma message("UESHADERMETADATA_VERSION 84C59757-69B5-48FB-A7AA-C168F63C2354")
+
 #include "../ShadingCommon.ush"
 #include "../OctahedralCommon.ush"
 #include "../SceneData.ush"
diff --git a/Engine/Shaders/Public/Platform/D3D/D3DCommon.ush b/Engine/Shaders/Public/Platform/D3D/D3DCommon.ush
--- a/Engine/Shaders/Public/Platform/D3D/D3DCommon.ush
+++ b/Engine/Shaders/Public/Platform/D3D/D3DCommon.ush
@@ -5,7 +5,7 @@ D3DCommon.ush: Common defines for the D3D shader compiler
 =============================================================================*/
 
 // Update this GUID to invalidate D3D shaders and force a shader recompilation for D3D RHIs.
-// GUID: EE1A1608-B3B7-4239-9B61-4DF491A9B6E3
+#pragma message("UESHADERMETADATA_VERSION 9C818FFC-0B81-4BCB-BF16-CFA7E4094DE3")
 
 #include "/Engine/Shared/HLSLReservedSpaces.h"
 
diff --git a/Engine/Shaders/Public/Platform/GL/GLCommon.ush b/Engine/Shaders/Public/Platform/GL/GLCommon.ush
--- a/Engine/Shaders/Public/Platform/GL/GLCommon.ush
+++ b/Engine/Shaders/Public/Platform/GL/GLCommon.ush
@@ -4,9 +4,9 @@
 	GLCommon.ush: Common GL shader code
 =============================================================================*/
 
-// Update this GUID to improve shader recompilation for GL only shaders
-// GUID = F1338239-A231-40B9-AAD4-421C7AD5306A
-
 #pragma once
 
+// Update this GUID to invalidate shader recompilation for only GL shaders
+#pragma message("UESHADERMETADATA_VERSION 1535AA07-356D-4F56-88BD-B38BDD770D36")
+
 
diff --git a/Engine/Shaders/Public/Platform/Metal/MetalCommon.ush b/Engine/Shaders/Public/Platform/Metal/MetalCommon.ush
--- a/Engine/Shaders/Public/Platform/Metal/MetalCommon.ush
+++ b/Engine/Shaders/Public/Platform/Metal/MetalCommon.ush
@@ -4,11 +4,11 @@
 	MetalCommon.usf: Common Metal shader code
 =============================================================================*/
 
-// Update this GUID to improve shader recompilation for Metal only shaders
-// GUID = 1E938C6C-E047-46A0-A37A-CA31AD446D64
-
 #pragma once
 
+// Update this GUID to invalidate shader recompilation for only Metal shaders
+#pragma message("UESHADERMETADATA_VERSION 059503F5-57EB-4E3E-96DB-63F381FC9FEC")
+
 #if VERTEXSHADER && HULLSHADER
 #ifndef TESSELLATIONSHADER
 #define TESSELLATIONSHADER 1
diff --git a/Engine/Shaders/Public/Platform/Vulkan/VulkanCommon.ush b/Engine/Shaders/Public/Platform/Vulkan/VulkanCommon.ush
--- a/Engine/Shaders/Public/Platform/Vulkan/VulkanCommon.ush
+++ b/Engine/Shaders/Public/Platform/Vulkan/VulkanCommon.ush
@@ -4,11 +4,10 @@
 	VulkanCommon.ush: Common Vulkan shader code
 =============================================================================*/
 
-// Update this GUID to improve shader recompilation for Vulkan only shaders
-// GUID = 8931F79A-3E7B-4BB4-BB26-341699B0D77E
-
 #pragma once
 
+// Update this GUID to invalidate shader recompilation for only Vulkan shaders
+#pragma message("UESHADERMETADATA_VERSION BB707B6A-C217-4B23-9A22-086E9EC77FCF")
 
 #define COMPILER_SUPPORTS_ULONG_TYPES 1
 #define COMPILER_SUPPORTS_UINT64_IMAGE_ATOMICS (1 && COMPILER_SUPPORTS_ULONG_TYPES)
diff --git a/Engine/Source/Developer/Apple/MetalShaderFormat/Private/MetalShaderFormat.cpp b/Engine/Source/Developer/Apple/MetalShaderFormat/Private/MetalShaderFormat.cpp
--- a/Engine/Source/Developer/Apple/MetalShaderFormat/Private/MetalShaderFormat.cpp
+++ b/Engine/Source/Developer/Apple/MetalShaderFormat/Private/MetalShaderFormat.cpp
@@ -36,20 +36,12 @@ extern bool StripShader_Metal(TArray<uint8>& Code, class FString const& DebugPat
 extern uint64 AppendShader_Metal(class FString const& ArchivePath, const FSHAHash& Hash, TArray<uint8>& Code);
 extern bool FinalizeLibrary_Metal(class FName const& Format, class FString const& ArchivePath, class FString const& LibraryPath, TSet<uint64> const& Shaders, class FString const& DebugOutputDir);
 
+/** Version for shader format, this becomes part of the DDC key. */
+static const FGuid UE_SHADER_METAL_VER = FGuid("FB0EA082-A9FC-469F-AE8F-03CFE0418EA5");
+
 class FMetalShaderFormat : public IShaderFormat
 {
 public:
-	enum
-	{
-		HEADER_VERSION = 72,
-	};
-	
-	struct FVersion
-	{
-		uint16 XcodeVersion;
-		uint16 HLSLCCMinor		: 8;
-		uint16 Format			: 8;
-	};
 	FMetalShaderFormat()
 	{
 		FMetalCompilerToolchain::CreateAndInit();
@@ -60,13 +52,6 @@ public:
 	}
 	virtual uint32 GetVersion(FName Format) const override final
 	{
-		static_assert(sizeof(FMetalShaderFormat::FVersion) == sizeof(uint32), "Out of bits!");
-		union
-		{
-			FMetalShaderFormat::FVersion Version;
-			uint32 Raw;
-		} Version;
-
 		// If there's no compiler on this machine, this is irrelevant so just return 0
 		if (!FMetalCompilerToolchain::Get()->IsCompilerAvailable())
 		{
@@ -100,16 +85,10 @@ public:
 			// Since Metal minor/patch version changes every Xcode minor version, we don't want users to rebuild shaders for every minor version update
 		}
 
-		Version.Version.XcodeVersion = HashValue;
-		Version.Version.Format = FMetalShaderFormat::HEADER_VERSION;
-		Version.Version.HLSLCCMinor = HLSLCC_VersionMinor;
-
-		// Check that we didn't overwrite any bits
-		check(Version.Version.XcodeVersion == HashValue);
-		check(Version.Version.Format == FMetalShaderFormat::HEADER_VERSION);
-		check(Version.Version.HLSLCCMinor == HLSLCC_VersionMinor);
+		uint32 Result = GetTypeHash(HashValue);
 
-		uint32 Result = Version.Raw;
+		Result = HashCombine(Result, GetTypeHash(HLSLCC_VersionMinor));
+		Result = HashCombine(Result, GetTypeHash(UE_SHADER_METAL_VER));
 
 #if UE_METAL_SHADER_COMPILER_ALLOW_DEAD_CODE_REMOVAL
 		Result = HashCombine(Result, 0x75E2FE85);
diff --git a/Engine/Source/Developer/ShaderFormatOpenGL/Private/ShaderFormatOpenGL.cpp b/Engine/Source/Developer/ShaderFormatOpenGL/Private/ShaderFormatOpenGL.cpp
--- a/Engine/Source/Developer/ShaderFormatOpenGL/Private/ShaderFormatOpenGL.cpp
+++ b/Engine/Source/Developer/ShaderFormatOpenGL/Private/ShaderFormatOpenGL.cpp
@@ -28,14 +28,11 @@ extern void CompileOpenGLShader(
 	const FString& WorkingDirectory,
 	GLSLVersion Version);
 
+/** Version for shader format, this becomes part of the DDC key. */
+static const FGuid UE_SHADER_GLSL_VER = FGuid("45CB4440-A342-475E-AAD1-4CB8E228D1C9");
+
 class FShaderFormatGLSL : public UE::ShaderCompilerCommon::FBaseShaderFormat 
 {
-	enum
-	{
-		/** Version for shader format, this becomes part of the DDC key. */
-		UE_SHADER_GLSL_VER = 107,
-	};
-
 	static void CheckFormat(FName Format)
 	{
 		check(Format == NAME_GLSL_150_ES3_1 || Format == NAME_GLSL_ES3_1_ANDROID);
@@ -45,17 +42,16 @@ public:
 	virtual uint32 GetVersion(FName Format) const override
 	{
 		CheckFormat(Format);
-		uint32 GLSLVersion = 0;
-		if (Format == NAME_GLSL_150_ES3_1 || Format == NAME_GLSL_ES3_1_ANDROID)
+		FGuid GLSLVersion{};
+		if (ensure(Format == NAME_GLSL_150_ES3_1 || Format == NAME_GLSL_ES3_1_ANDROID))
 		{
 			GLSLVersion = UE_SHADER_GLSL_VER;
 		}
-		else
-		{
-			check(0);
-		}
 
-		uint32 Version = ((HLSLCC_VersionMinor & 0xff) << 8) | (GLSLVersion & 0xff);
+		const uint32 BaseHash = GetTypeHash(GLSLVersion);
+
+		uint32 Version = GetTypeHash(HLSLCC_VersionMinor);
+		Version = HashCombine(Version, BaseHash);
 
 	#if UE_OPENGL_SHADER_COMPILER_ALLOW_DEAD_CODE_REMOVAL
 		Version = HashCombine(Version, 0x75E2FE85);
diff --git a/Engine/Source/Developer/VulkanShaderFormat/Private/VulkanShaderFormat.cpp b/Engine/Source/Developer/VulkanShaderFormat/Private/VulkanShaderFormat.cpp
--- a/Engine/Source/Developer/VulkanShaderFormat/Private/VulkanShaderFormat.cpp
+++ b/Engine/Source/Developer/VulkanShaderFormat/Private/VulkanShaderFormat.cpp
@@ -30,33 +30,32 @@ extern void OutputVulkanDebugData(
 	const FShaderPreprocessOutput& PreprocessOutput, 
 	const FShaderCompilerOutput& Output);
 
+static const FGuid UE_SHADER_VULKAN_ES3_1_VER = FGuid("6D333063-D2F7-4AA5-A79D-BA73F32C898E");
+static const FGuid UE_SHADER_VULKAN_SM5_VER = FGuid("6EC81E81-BDE8-4F09-8A70-AD09A817B32A");
+static const FGuid UE_SHADER_VULKAN_SM6_VER = FGuid("C732FBB7-4CAD-4249-A760-912777047233");
+
 class FShaderFormatVulkan : public IShaderFormat
 {
-	enum 
-	{
-		UE_SHADER_VULKAN_ES3_1_VER	= 37,
-		UE_SHADER_VULKAN_SM5_VER 	= 37,
-		UE_SHADER_VULKAN_SM6_VER	= 37,
-	};
-
-	int32 InternalGetVersion(FName Format) const
+	FGuid InternalGetVersion(FName Format) const
 	{
 		if (Format == NAME_VULKAN_SM6)
 		{
 			return UE_SHADER_VULKAN_SM6_VER;
 		}
-		else if (Format == NAME_VULKAN_SM5 || Format == NAME_VULKAN_SM5_ANDROID)
+
+		if (Format == NAME_VULKAN_SM5 || Format == NAME_VULKAN_SM5_ANDROID)
 		{
 			return UE_SHADER_VULKAN_SM5_VER;
 		}
-		else if (Format == NAME_VULKAN_ES3_1_ANDROID || Format == NAME_VULKAN_ES3_1)
+
+		if (Format == NAME_VULKAN_ES3_1_ANDROID || Format == NAME_VULKAN_ES3_1)
 		{
 			return UE_SHADER_VULKAN_ES3_1_VER;
 		}
 
 		FString FormatStr = Format.ToString();
 		checkf(0, TEXT("Invalid shader format passed to Vulkan shader compiler: %s"), *FormatStr);
-		return -1;
+		return FGuid();
 	}
 
 	uint32 ShaderConductorVersionHash;
@@ -95,14 +94,11 @@ public:
 
 	virtual bool PreprocessShader(const FShaderCompilerInput& Input, const FShaderCompilerEnvironment& Environment, FShaderPreprocessOutput& PreprocessOutput) const
 	{
-		check(InternalGetVersion(Input.ShaderFormat) >= 0);
-
 		return PreprocessVulkanShader(Input, Environment, PreprocessOutput);
 	}
 
 	virtual void CompilePreprocessedShader(const FShaderCompilerInput& Input, const FShaderPreprocessOutput& PreprocessOutput, FShaderCompilerOutput& Output,const FString& WorkingDirectory) const override
 	{
-		check(InternalGetVersion(Input.ShaderFormat) >= 0);
 		CompileVulkanShader(Input, PreprocessOutput.GetSource(), Output, WorkingDirectory);
 	}
 
