[shaders] fix IOS path to not assume ShaderConductor ignores nulls when loading source; I don't think this is an obvious or desirable choice in the API (despite the fact that changing this behaviour in my SC refactoring was unintentional)

diff --git a/Engine/Source/Developer/Apple/MetalShaderFormat/Private/MetalDerivedData.cpp b/Engine/Source/Developer/Apple/MetalShaderFormat/Private/MetalDerivedData.cpp
--- a/Engine/Source/Developer/Apple/MetalShaderFormat/Private/MetalDerivedData.cpp
+++ b/Engine/Source/Developer/Apple/MetalShaderFormat/Private/MetalDerivedData.cpp
@@ -310,7 +310,7 @@ void DoCompileMetalShader(
 		// If source data was modified, reload it into the compiler context
 		if (bSourceDataWasModified)
 		{
-			CompilerContext.LoadSource(SourceData.c_str(), /*Fileanme:*/ nullptr, /*EntryPoint:*/ nullptr, Frequency);
+			CompilerContext.LoadSource(FAnsiStringView(SourceData.c_str(), SourceData.length()), Input.VirtualSourceFilePath, Input.EntryPointName, Frequency);
 		}
 
 		if (bDumpDebugInfo)
