Remove Compression.h from CoreMinimal.h
It's about 1/4 of CoreMinimal.h but rarely needed (Compression.h pulls on CriticalSection.h and Map.h that are costly).

diff --git a/Engine/Plugins/Developer/Concert/ConcertMain/Source/Concert/Private/ConcertMessageData.cpp b/Engine/Plugins/Developer/Concert/ConcertMain/Source/Concert/Private/ConcertMessageData.cpp
--- a/Engine/Plugins/Developer/Concert/ConcertMain/Source/Concert/Private/ConcertMessageData.cpp
+++ b/Engine/Plugins/Developer/Concert/ConcertMain/Source/Concert/Private/ConcertMessageData.cpp
@@ -6,6 +6,7 @@
 #include "IdentifierTable/ConcertTransportArchives.h"
 
 #include "Misc/App.h"
+#include "Misc/Compression.h"
 #include "UObject/StructOnScope.h"
 
 #include "StructSerializer.h"
diff --git a/Engine/Plugins/Experimental/TextureGraph/Source/TextureGraphEngine/Data/RawBuffer.cpp b/Engine/Plugins/Experimental/TextureGraph/Source/TextureGraphEngine/Data/RawBuffer.cpp
--- a/Engine/Plugins/Experimental/TextureGraph/Source/TextureGraphEngine/Data/RawBuffer.cpp
+++ b/Engine/Plugins/Experimental/TextureGraph/Source/TextureGraphEngine/Data/RawBuffer.cpp
@@ -4,6 +4,7 @@
 #include <IImageWrapperModule.h>
 #include <IImageWrapper.h>
 #include "Helper/Util.h"
+#include "Misc/Compression.h"
 #include "Misc/Paths.h"
 #include "Modules/ModuleManager.h"
 #include "Async/ParallelFor.h"
diff --git a/Engine/Source/Developer/Apple/MetalShaderFormat/Private/MetalShaderCompiler.cpp b/Engine/Source/Developer/Apple/MetalShaderFormat/Private/MetalShaderCompiler.cpp
--- a/Engine/Source/Developer/Apple/MetalShaderFormat/Private/MetalShaderCompiler.cpp
+++ b/Engine/Source/Developer/Apple/MetalShaderFormat/Private/MetalShaderCompiler.cpp
@@ -1,10 +1,13 @@
 // Copyright Epic Games, Inc. All Rights Reserved.
 
+#if UE_ENABLE_INCLUDE_ORDER_DEPRECATED_IN_5_4
 #include "CoreMinimal.h"
+#endif
 #include "HAL/PlatformFileManager.h"
 #include "MetalShaderFormat.h"
 #include "MetalShaderResources.h"
 #include "Misc/ConfigCacheIni.h"
+#include "Misc/Compression.h"
 #include "Misc/EngineVersion.h"
 #include "Misc/FileHelper.h"
 #include "Misc/Paths.h"
diff --git a/Engine/Source/Developer/DerivedDataCache/Private/S3CacheStore.cpp b/Engine/Source/Developer/DerivedDataCache/Private/S3CacheStore.cpp
--- a/Engine/Source/Developer/DerivedDataCache/Private/S3CacheStore.cpp
+++ b/Engine/Source/Developer/DerivedDataCache/Private/S3CacheStore.cpp
@@ -24,6 +24,7 @@
 #include "HashingArchiveProxy.h"
 #include "Memory/SharedBuffer.h"
 #include "Misc/Base64.h"
+#include "Misc/Compression.h"
 #include "Misc/ConfigCacheIni.h"
 #include "Misc/FeedbackContext.h"
 #include "Misc/FileHelper.h"
diff --git a/Engine/Source/Developer/Horde/Private/Storage/Bundles/BundleCompression.cpp b/Engine/Source/Developer/Horde/Private/Storage/Bundles/BundleCompression.cpp
--- a/Engine/Source/Developer/Horde/Private/Storage/Bundles/BundleCompression.cpp
+++ b/Engine/Source/Developer/Horde/Private/Storage/Bundles/BundleCompression.cpp
@@ -1,6 +1,7 @@
 // Copyright Epic Games, Inc. All Rights Reserved.
 
 #include "Storage/Bundles/BundleCompression.h"
+#include "Misc/Compression.h"
 #include "../../HordePlatform.h"
 
 size_t FBundleCompression::GetMaxSize(EBundleCompressionFormat Format, const FMemoryView& Input)
diff --git a/Engine/Source/Developer/ShaderCompilerCommon/Private/ShaderCompilerCommon.cpp b/Engine/Source/Developer/ShaderCompilerCommon/Private/ShaderCompilerCommon.cpp
--- a/Engine/Source/Developer/ShaderCompilerCommon/Private/ShaderCompilerCommon.cpp
+++ b/Engine/Source/Developer/ShaderCompilerCommon/Private/ShaderCompilerCommon.cpp
@@ -3,6 +3,7 @@
 #include "ShaderCompilerCommon.h"
 #include "ShaderParameterParser.h"
 #include "Misc/Base64.h"
+#include "Misc/Compression.h"
 #include "Misc/FileHelper.h"
 #include "Misc/Paths.h"
 #include "Misc/PathViews.h"
diff --git a/Engine/Source/Developer/SlateReflector/Private/Widgets/SWidgetSnapshotVisualizer.cpp b/Engine/Source/Developer/SlateReflector/Private/Widgets/SWidgetSnapshotVisualizer.cpp
--- a/Engine/Source/Developer/SlateReflector/Private/Widgets/SWidgetSnapshotVisualizer.cpp
+++ b/Engine/Source/Developer/SlateReflector/Private/Widgets/SWidgetSnapshotVisualizer.cpp
@@ -16,6 +16,7 @@
 #include "Framework/Layout/ScrollyZoomy.h"
 #include "Framework/MultiBox/MultiBoxBuilder.h"
 #include "Misc/Base64.h"
+#include "Misc/Compression.h"
 
 #include "SlateNavigationEventSimulator.h"
 #include "SlateReflectorModule.h"
diff --git a/Engine/Source/Editor/UnrealEd/Private/Factories/PackFactory.cpp b/Engine/Source/Editor/UnrealEd/Private/Factories/PackFactory.cpp
--- a/Engine/Source/Editor/UnrealEd/Private/Factories/PackFactory.cpp
+++ b/Engine/Source/Editor/UnrealEd/Private/Factories/PackFactory.cpp
@@ -13,6 +13,7 @@
 #include "Misc/Paths.h"
 #include "Serialization/MemoryWriter.h"
 #include "Serialization/MemoryReader.h"
+#include "Misc/Compression.h"
 #include "Misc/ConfigCacheIni.h"
 #include "Misc/ConfigContext.h"
 #include "Misc/FeedbackContext.h"
diff --git a/Engine/Source/Programs/ShaderCompileWorker/Private/ShaderCompileWorker.cpp b/Engine/Source/Programs/ShaderCompileWorker/Private/ShaderCompileWorker.cpp
--- a/Engine/Source/Programs/ShaderCompileWorker/Private/ShaderCompileWorker.cpp
+++ b/Engine/Source/Programs/ShaderCompileWorker/Private/ShaderCompileWorker.cpp
@@ -4,7 +4,10 @@
 // ShaderCompileWorker.cpp : Defines the entry point for the console application.
 //
 
+#if UE_ENABLE_INCLUDE_ORDER_DEPRECATED_IN_5_4
 #include "CoreMinimal.h"
+#endif
+#include "Misc/Compression.h"
 #include "RequiredProgramMainCPPInclude.h"
 #include "ShaderCompilerCore.h"
 #include "ShaderCompilerCommon.h"
diff --git a/Engine/Source/Runtime/AIModule/Private/EnvironmentQuery/EnvQueryDebugHelpers.cpp b/Engine/Source/Runtime/AIModule/Private/EnvironmentQuery/EnvQueryDebugHelpers.cpp
--- a/Engine/Source/Runtime/AIModule/Private/EnvironmentQuery/EnvQueryDebugHelpers.cpp
+++ b/Engine/Source/Runtime/AIModule/Private/EnvironmentQuery/EnvQueryDebugHelpers.cpp
@@ -7,6 +7,7 @@
 #include "EnvironmentQuery/EnvQueryTest.h"
 #include "VisualLoggerExtension.h"
 #include "EnvironmentQuery/EQSRenderingComponent.h"
+#include "Misc/Compression.h"
 
 #include UE_INLINE_GENERATED_CPP_BY_NAME(EnvQueryDebugHelpers)
 
diff --git a/Engine/Source/Runtime/Core/Public/CoreMinimal.h b/Engine/Source/Runtime/Core/Public/CoreMinimal.h
--- a/Engine/Source/Runtime/Core/Public/CoreMinimal.h
+++ b/Engine/Source/Runtime/Core/Public/CoreMinimal.h
@@ -53,7 +53,9 @@
 #include "Containers/ContainerAllocationPolicies.h"
 #include "Templates/IsEnumClass.h"
 #include "HAL/PlatformProperties.h"
+#if UE_ENABLE_INCLUDE_ORDER_DEPRECATED_IN_5_4
 #include "Misc/Compression.h"
+#endif
 #include "Misc/EngineVersionBase.h"
 #include "Internationalization/TextNamespaceFwd.h"
 #include "Serialization/Archive.h"
diff --git a/Engine/Source/Runtime/CoreUObject/Private/UObject/PackageFileSummary.cpp b/Engine/Source/Runtime/CoreUObject/Private/UObject/PackageFileSummary.cpp
--- a/Engine/Source/Runtime/CoreUObject/Private/UObject/PackageFileSummary.cpp
+++ b/Engine/Source/Runtime/CoreUObject/Private/UObject/PackageFileSummary.cpp
@@ -1,6 +1,7 @@
 // Copyright Epic Games, Inc. All Rights Reserved.
 
 #include "UObject/PackageFileSummary.h"
+#include "Misc/Compression.h"
 #include "UObject/Linker.h"
 #include "Serialization/StructuredArchive.h"
 #include "UObject/UObjectGlobals.h"
diff --git a/Engine/Source/Runtime/CrashReportCore/Private/CrashUpload.cpp b/Engine/Source/Runtime/CrashReportCore/Private/CrashUpload.cpp
--- a/Engine/Source/Runtime/CrashReportCore/Private/CrashUpload.cpp
+++ b/Engine/Source/Runtime/CrashReportCore/Private/CrashUpload.cpp
@@ -3,6 +3,7 @@
 #include "CrashUpload.h"
 #include "CrashReportCoreModule.h"
 #include "HAL/FileManager.h"
+#include "Misc/Compression.h"
 #include "Misc/FileHelper.h"
 #include "Internationalization/Internationalization.h"
 #include "Misc/Guid.h"
diff --git a/Engine/Source/Runtime/Datasmith/DatasmithCore/Private/DatasmithMeshSerialization.cpp b/Engine/Source/Runtime/Datasmith/DatasmithCore/Private/DatasmithMeshSerialization.cpp
--- a/Engine/Source/Runtime/Datasmith/DatasmithCore/Private/DatasmithMeshSerialization.cpp
+++ b/Engine/Source/Runtime/Datasmith/DatasmithCore/Private/DatasmithMeshSerialization.cpp
@@ -5,6 +5,7 @@
 
 #include "Compression/OodleDataCompression.h"
 #include "HAL/FileManager.h"
+#include "Misc/Compression.h"
 #include "Serialization/MemoryReader.h"
 #include "Serialization/MemoryWriter.h"
 #include "StaticMeshAttributes.h"
diff --git a/Engine/Source/Runtime/GameplayDebugger/Private/GameplayDebuggerTypes.cpp b/Engine/Source/Runtime/GameplayDebugger/Private/GameplayDebuggerTypes.cpp
--- a/Engine/Source/Runtime/GameplayDebugger/Private/GameplayDebuggerTypes.cpp
+++ b/Engine/Source/Runtime/GameplayDebugger/Private/GameplayDebuggerTypes.cpp
@@ -2,6 +2,7 @@
 
 #include "GameplayDebuggerTypes.h"
 #include "InputCoreTypes.h"
+#include "Misc/Compression.h"
 #include "SceneView.h"
 #include "Serialization/MemoryWriter.h"
 #include "Serialization/MemoryReader.h"
diff --git a/Engine/Source/Runtime/Online/BuildPatchServices/Private/Data/ChunkData.cpp b/Engine/Source/Runtime/Online/BuildPatchServices/Private/Data/ChunkData.cpp
--- a/Engine/Source/Runtime/Online/BuildPatchServices/Private/Data/ChunkData.cpp
+++ b/Engine/Source/Runtime/Online/BuildPatchServices/Private/Data/ChunkData.cpp
@@ -4,6 +4,7 @@
 #include "Serialization/MemoryReader.h"
 #include "Serialization/MemoryWriter.h"
 #include "Misc/CommandLine.h"
+#include "Misc/Compression.h"
 #include "Common/FileSystem.h"
 #include "BuildPatchHash.h"
 #include "BuildPatchManifest.h"
diff --git a/Engine/Source/Runtime/Online/BuildPatchServices/Private/Data/ManifestData.cpp b/Engine/Source/Runtime/Online/BuildPatchServices/Private/Data/ManifestData.cpp
--- a/Engine/Source/Runtime/Online/BuildPatchServices/Private/Data/ManifestData.cpp
+++ b/Engine/Source/Runtime/Online/BuildPatchServices/Private/Data/ManifestData.cpp
@@ -1,6 +1,7 @@
 // Copyright Epic Games, Inc. All Rights Reserved.
 
 #include "Data/ManifestData.h"
+#include "Misc/Compression.h"
 #include "Misc/EnumClassFlags.h"
 #include "Serialization/Archive.h"
 #include "Serialization/MemoryWriter.h"
diff --git a/Engine/Source/Runtime/Online/BuildPatchServices/Private/Data/ManifestUObject.cpp b/Engine/Source/Runtime/Online/BuildPatchServices/Private/Data/ManifestUObject.cpp
--- a/Engine/Source/Runtime/Online/BuildPatchServices/Private/Data/ManifestUObject.cpp
+++ b/Engine/Source/Runtime/Online/BuildPatchServices/Private/Data/ManifestUObject.cpp
@@ -8,6 +8,7 @@
 #include "Algo/Transform.h"
 #include "Data/ManifestData.h"
 #include "BuildPatchManifest.h"
+#include "Misc/Compression.h"
 
 #include UE_INLINE_GENERATED_CPP_BY_NAME(ManifestUObject)
 
diff --git a/Engine/Source/Runtime/OpenGLDrv/Private/OpenGLBinaryProgramUtils.cpp b/Engine/Source/Runtime/OpenGLDrv/Private/OpenGLBinaryProgramUtils.cpp
--- a/Engine/Source/Runtime/OpenGLDrv/Private/OpenGLBinaryProgramUtils.cpp
+++ b/Engine/Source/Runtime/OpenGLDrv/Private/OpenGLBinaryProgramUtils.cpp
@@ -4,6 +4,7 @@
 	OpenGLBinaryProgramUtils.cpp
 =============================================================================*/
 
+#include "Misc/Compression.h"
 #include "OpenGLBinaryProgramUtils.h"
 
 namespace UE
diff --git a/Engine/Source/Runtime/OpenGLDrv/Private/OpenGLShaders.cpp b/Engine/Source/Runtime/OpenGLDrv/Private/OpenGLShaders.cpp
--- a/Engine/Source/Runtime/OpenGLDrv/Private/OpenGLShaders.cpp
+++ b/Engine/Source/Runtime/OpenGLDrv/Private/OpenGLShaders.cpp
@@ -8,6 +8,7 @@
 #include "HAL/PlatformFileManager.h"
 #include "HAL/FileManager.h"
 #include "Misc/Paths.h"
+#include "Misc/Compression.h"
 #include "Serialization/MemoryWriter.h"
 #include "Serialization/MemoryReader.h"
 #include "OpenGLDrvPrivate.h"
diff --git a/Engine/Source/Runtime/Renderer/Private/SceneOcclusion.cpp b/Engine/Source/Runtime/Renderer/Private/SceneOcclusion.cpp
--- a/Engine/Source/Runtime/Renderer/Private/SceneOcclusion.cpp
+++ b/Engine/Source/Runtime/Renderer/Private/SceneOcclusion.cpp
@@ -22,6 +22,7 @@
 #include "PixelShaderUtils.h"
 #include "RenderCore.h"
 #include "MobileBasePassRendering.h"
+#include "Misc/Compression.h"
 
 
 /*-----------------------------------------------------------------------------
