Added Texture atomic array flattening for TSR to SM5

diff --git a/Engine/Source/Developer/Apple/MetalShaderFormat/Private/MetalDerivedData.cpp b/Engine/Source/Developer/Apple/MetalShaderFormat/Private/MetalDerivedData.cpp
--- a/Engine/Source/Developer/Apple/MetalShaderFormat/Private/MetalDerivedData.cpp
+++ b/Engine/Source/Developer/Apple/MetalShaderFormat/Private/MetalDerivedData.cpp
@@ -988,8 +988,13 @@ void DoCompileMetalShader(
                 // Must be done as VSM uses 2DArray and requires atomics support. And Metal does not
                 // support atomics on 2Darray...
 				TargetDesc.CompileFlags->SetDefine(TEXT("flatten_2d_array"), 1);
-                TargetDesc.CompileFlags->SetDefine(TEXT("flatten_2d_array_names"), TEXT("VirtualShadowMap_PhysicalPagePool,OutPhysicalPagePool,OutDepthBufferArray,PhysicalPagePool,ShadowDepthPass_OutDepthBufferArray"));
+                TargetDesc.CompileFlags->SetDefine(TEXT("flatten_2d_array_names"), TEXT("VirtualShadowMap_PhysicalPagePool,OutPhysicalPagePool,OutDepthBufferArray,PhysicalPagePool,ShadowDepthPass_OutDepthBufferArray,PrevAtomicTextureArray,PrevAtomicOutput"));
             }
+			else if(Input.Target.Platform == SP_METAL_SM5)
+			{
+				TargetDesc.CompileFlags->SetDefine(TEXT("flatten_2d_array"), 1);
+				TargetDesc.CompileFlags->SetDefine(TEXT("flatten_2d_array_names"), TEXT("PrevAtomicTextureArray,PrevAtomicOutput"));
+			}
 
 			switch (Semantics)
 			{
diff --git a/Engine/Source/Runtime/Apple/MetalRHI/Private/MetalTexture.cpp b/Engine/Source/Runtime/Apple/MetalRHI/Private/MetalTexture.cpp
--- a/Engine/Source/Runtime/Apple/MetalRHI/Private/MetalTexture.cpp
+++ b/Engine/Source/Runtime/Apple/MetalRHI/Private/MetalTexture.cpp
@@ -613,7 +613,7 @@ FMetalSurface::FMetalSurface(FRHICommandListBase* RHICmdList, FMetalTextureCreat
 		}
         else if (bTextureArrayWithAtomics)
         {
-            checkf(GMaxRHIFeatureLevel == ERHIFeatureLevel::SM6, TEXT("Requested texture array with atomics that is unsupported on this platform"));
+            checkf(GMaxRHIFeatureLevel >= ERHIFeatureLevel::SM5, TEXT("Requested texture array with atomics that is unsupported on this platform"));
             
             MTL::Device* Device = GetMetalDeviceContext().GetDevice();
 
