Move the FShaderResourceTable and StaticSlots arrays from platform RHI implementations into the base FRHIShader type
 - This data is required by uniform binding code, and was copy/paste across platform RHIs. Moving it to the base RHI type will allow the RHI validation layer to enumerate resources in uniform buffers, which have been opaque up until now.
 - FShaderResourceTable has moved from RenderCore to RHI.
 - UE::RHICore::ApplyStaticUniformBuffers and UE::RHICore::SetResourcesFromTables now only need the shader, since the binding info is stored within it.
 - The serializer functions in shader format modules and platform RHIs have been fixed up to handle the SRT being in the base type. The data format has not changed, so no shader versions need bumping.
 - Removed unnecessary operator == and GetTypeHash functions in some places.

Includes all platforms / RHIs except OpenGL, which will follow in another CL.

diff --git a/Engine/Source/Developer/Apple/MetalShaderFormat/Private/MetalShaderCompiler.cpp b/Engine/Source/Developer/Apple/MetalShaderFormat/Private/MetalShaderCompiler.cpp
--- a/Engine/Source/Developer/Apple/MetalShaderFormat/Private/MetalShaderCompiler.cpp
+++ b/Engine/Source/Developer/Apple/MetalShaderFormat/Private/MetalShaderCompiler.cpp
@@ -291,6 +291,8 @@ void BuildMetalShaderOutput(
 	}
 
 	FMetalCodeHeader Header;
+	FShaderResourceTable SRT;
+
 	Header.CompileFlags = (ShaderInput.Environment.CompilerFlags.Contains(CFLAG_Debug) ? (1 << CFLAG_Debug) : 0);
 	Header.CompileFlags |= (bNoFastMath ? (1 << CFLAG_NoFastMath) : 0);
 	Header.CompileFlags |= (ShaderInput.Environment.CompilerFlags.Contains(CFLAG_ExtraShaderData) ? (1 << CFLAG_ExtraShaderData) : 0);
@@ -584,15 +586,15 @@ void BuildMetalShaderOutput(
 		CullGlobalUniformBuffers(ShaderInput.Environment.UniformBufferMap, ShaderOutput.ParameterMap);
 
 		// Copy over the bits indicating which resource tables are active.
-		Header.Bindings.ShaderResourceTable.ResourceTableBits = GenericSRT.ResourceTableBits;
+		SRT.ResourceTableBits = GenericSRT.ResourceTableBits;
 
-		Header.Bindings.ShaderResourceTable.ResourceTableLayoutHashes = GenericSRT.ResourceTableLayoutHashes;
+		SRT.ResourceTableLayoutHashes = GenericSRT.ResourceTableLayoutHashes;
 
 		// Now build our token streams.
-		BuildResourceTableTokenStream(GenericSRT.TextureMap, GenericSRT.MaxBoundResourceTable, Header.Bindings.ShaderResourceTable.TextureMap);
-		BuildResourceTableTokenStream(GenericSRT.ShaderResourceViewMap, GenericSRT.MaxBoundResourceTable, Header.Bindings.ShaderResourceTable.ShaderResourceViewMap);
-		BuildResourceTableTokenStream(GenericSRT.SamplerMap, GenericSRT.MaxBoundResourceTable, Header.Bindings.ShaderResourceTable.SamplerMap);
-		BuildResourceTableTokenStream(GenericSRT.UnorderedAccessViewMap, GenericSRT.MaxBoundResourceTable, Header.Bindings.ShaderResourceTable.UnorderedAccessViewMap);
+		BuildResourceTableTokenStream(GenericSRT.TextureMap, GenericSRT.MaxBoundResourceTable, SRT.TextureMap);
+		BuildResourceTableTokenStream(GenericSRT.ShaderResourceViewMap, GenericSRT.MaxBoundResourceTable, SRT.ShaderResourceViewMap);
+		BuildResourceTableTokenStream(GenericSRT.SamplerMap, GenericSRT.MaxBoundResourceTable, SRT.SamplerMap);
+		BuildResourceTableTokenStream(GenericSRT.UnorderedAccessViewMap, GenericSRT.MaxBoundResourceTable, SRT.UnorderedAccessViewMap);
 
 		Header.Bindings.NumUniformBuffers = FMath::Max((uint8)GetNumUniformBuffersUsed(GenericSRT), Header.Bindings.NumUniformBuffers);
 	}
@@ -627,7 +629,7 @@ void BuildMetalShaderOutput(
 		FMemoryWriter Ar(ShaderOutput.ShaderCode.GetWriteAccess(), true);
 		uint8 PrecompiledFlag = 0;
 		Ar << PrecompiledFlag;
-		Ar << Header;
+		Header.Serialize(Ar, SRT);
 		Ar.Serialize((void*)USFSource, SourceLen + 1 - (USFSource - InShaderSource));
 		
 		ShaderOutput.ModifiedShaderSource = MetalCode;
@@ -794,7 +796,7 @@ void BuildMetalShaderOutput(
 			FMemoryWriter Ar(ShaderOutput.ShaderCode.GetWriteAccess(), true);
 			uint8 PrecompiledFlag = 1;
 			Ar << PrecompiledFlag;
-			Ar << Header;
+            Header.Serialize(Ar, SRT);
 
 			// jam it into the output bytes
 			Ar.Serialize(Bytecode.OutputFile.GetData(), Bytecode.OutputFile.Num());
@@ -1123,7 +1125,8 @@ bool StripShader_Metal(TArray<uint8>& Code, class FString const& DebugPath, bool
 	{
 		// get the header
 		FMetalCodeHeader Header;
-		Ar << Header;
+        FShaderResourceTable SRT;
+        Header.Serialize(Ar, SRT);
 		
 		const FString ShaderName = ShaderCode.FindOptionalData(FShaderCodeName::Key);
 
@@ -1192,7 +1195,7 @@ bool StripShader_Metal(TArray<uint8>& Code, class FString const& DebugPath, bool
 				FShaderCode NewCode;
 				FMemoryWriter NewAr(NewCode.GetWriteAccess(), true);
 				NewAr << OfflineCompiledFlag;
-				NewAr << Header;
+                Header.Serialize(NewAr, SRT);
 				
 				// jam it into the output bytes
 				NewAr.Serialize(SourceCode.GetData(), SourceCode.Num());
@@ -1231,8 +1234,9 @@ uint64 AppendShader_Metal(FString const& WorkingDir, const FSHAHash& Hash, TArra
 		if (OfflineCompiledFlag == 1)
 		{
 			// get the header
-			FMetalCodeHeader Header;
-			Ar << Header;
+            FMetalCodeHeader Header;
+            FShaderResourceTable SRT;
+            Header.Serialize(Ar, SRT);
 
 			const FString ShaderName = ShaderCode.FindOptionalData(FShaderCodeName::Key);
 
@@ -1288,7 +1292,7 @@ uint64 AppendShader_Metal(FString const& WorkingDir, const FSHAHash& Hash, TArra
 						FShaderCode NewCode;
 						FMemoryWriter NewAr(NewCode.GetWriteAccess(), true);
 						NewAr << OfflineCompiledFlag;
-						NewAr << Header;
+                        Header.Serialize(NewAr, SRT);
 						
 						InShaderCode = NewCode.GetReadAccess();
 						
diff --git a/Engine/Source/Runtime/Apple/MetalRHI/Private/MetalCommands.cpp b/Engine/Source/Runtime/Apple/MetalRHI/Private/MetalCommands.cpp
--- a/Engine/Source/Runtime/Apple/MetalRHI/Private/MetalCommands.cpp
+++ b/Engine/Source/Runtime/Apple/MetalRHI/Private/MetalCommands.cpp
@@ -101,7 +101,7 @@ void FMetalRHICommandContext::ApplyStaticUniformBuffers(TRHIShader* Shader)
 {
 	if (Shader)
 	{
-		UE::RHICore::ApplyStaticUniformBuffers(this, Shader, Shader->StaticSlots, Shader->Bindings.ShaderResourceTable.ResourceTableLayoutHashes, GlobalUniformBuffers);
+		UE::RHICore::ApplyStaticUniformBuffers(this, Shader, GlobalUniformBuffers);
 	}
 }
 
diff --git a/Engine/Source/Runtime/Apple/MetalRHI/Private/MetalStateCache.cpp b/Engine/Source/Runtime/Apple/MetalRHI/Private/MetalStateCache.cpp
--- a/Engine/Source/Runtime/Apple/MetalRHI/Private/MetalStateCache.cpp
+++ b/Engine/Source/Runtime/Apple/MetalRHI/Private/MetalStateCache.cpp
@@ -1963,7 +1963,6 @@ void FMetalStateCache::SetResourcesFromTables(ShaderType Shader, CrossCompiler::
 		UE::RHICore::SetResourcesFromTables(
 			  FUniformResourceBinder { *this, Frequency }
 			, *Shader
-			, Shader->Bindings.ShaderResourceTable
 			, DirtyUniformBuffers[Frequency]
 			, BoundUniformBuffers[Frequency]
 #if ENABLE_RHI_VALIDATION
diff --git a/Engine/Source/Runtime/Apple/MetalRHI/Private/Shaders/Types/Templates/MetalBaseShader.h b/Engine/Source/Runtime/Apple/MetalRHI/Private/Shaders/Types/Templates/MetalBaseShader.h
--- a/Engine/Source/Runtime/Apple/MetalRHI/Private/Shaders/Types/Templates/MetalBaseShader.h
+++ b/Engine/Source/Runtime/Apple/MetalRHI/Private/Shaders/Types/Templates/MetalBaseShader.h
@@ -87,9 +87,6 @@ public:
 	/* Tier1 Argument buffer bitmasks */
 	TMap<uint32, TBitArray<>> ArgumentBitmasks;
 
-	/* Uniform buffer static slots */
-	TArray<FUniformBufferStaticSlot> StaticSlots;
-
 	/** The binding for the buffer side-table if present */
 	int32 SideTableBinding = -1;
 
@@ -148,7 +145,7 @@ void TMetalBaseShader<BaseResourceType, ShaderType>::Init(TArrayView<const uint8
 	check(OfflineCompiledFlag == 0 || OfflineCompiledFlag == 1);
 
 	// get the header
-	Ar << Header;
+	Header.Serialize(Ar, BaseResourceType::ShaderResourceTable);
 
 	ValidateVersion(Header.Version);
 
@@ -380,7 +377,7 @@ void TMetalBaseShader<BaseResourceType, ShaderType>::Init(TArrayView<const uint8
 	UniformBuffersCopyInfo = Header.UniformBuffersCopyInfo;
 	SideTableBinding = Header.SideTable;
 
-	UE::RHICore::InitStaticUniformBufferSlots(StaticSlots, Bindings.ShaderResourceTable);
+	UE::RHICore::InitStaticUniformBufferSlots(this);
 
 #if RHI_INCLUDE_SHADER_DEBUG_DATA
     this->Debug.ShaderName = FString::Printf(TEXT("Main_%0.8x_%0.8x"), Header.SourceLen, Header.SourceCRC);
diff --git a/Engine/Source/Runtime/Apple/MetalRHI/Public/MetalShaderResources.h b/Engine/Source/Runtime/Apple/MetalRHI/Public/MetalShaderResources.h
--- a/Engine/Source/Runtime/Apple/MetalRHI/Public/MetalShaderResources.h
+++ b/Engine/Source/Runtime/Apple/MetalRHI/Public/MetalShaderResources.h
@@ -103,52 +103,40 @@ struct FMetalShaderBindings
 {
 	TArray<TArray<CrossCompiler::FPackedArrayInfo>>	PackedUniformBuffers;
 	TArray<CrossCompiler::FPackedArrayInfo>			PackedGlobalArrays;
-	FShaderResourceTable							ShaderResourceTable;
 	TMap<uint8, TArray<uint8>>						ArgumentBufferMasks;
 	CrossCompiler::FShaderBindingInOutMask			InOutMask;
     FString                                         IRConverterReflectionJSON;
-    uint32                                          RSNumCBVs;
-    uint32                                          OutputSizeVS;
-    uint32                                          MaxInputPrimitivesPerMeshThreadgroupGS;
+    uint32                                          RSNumCBVs = 0;
+    uint32                                          OutputSizeVS = 0;
+    uint32                                          MaxInputPrimitivesPerMeshThreadgroupGS = 0;
 
-	uint32 	ConstantBuffers;
-	uint32  ArgumentBuffers;
-	uint8	NumSamplers;
-	uint8	NumUniformBuffers;
-	uint8	NumUAVs;
-	bool	bDiscards;
+	uint32 	ConstantBuffers = 0;
+	uint32  ArgumentBuffers = 0;
+	uint8	NumSamplers = 0;
+	uint8	NumUniformBuffers = 0;
+	uint8	NumUAVs = 0;
+	bool	bDiscards = false;
 
-	FMetalShaderBindings() :
-        RSNumCBVs(0),
-        OutputSizeVS(0),
-        MaxInputPrimitivesPerMeshThreadgroupGS(0),
-		ConstantBuffers(0),
-		ArgumentBuffers(0),
-		NumSamplers(0),
-		NumUniformBuffers(0),
-		NumUAVs(0),
-		bDiscards(false)
-	{
-	}
+	inline FArchive& Serialize(FArchive& Ar, FShaderResourceTable& SRT);
 };
 
-inline FArchive& operator<<(FArchive& Ar, FMetalShaderBindings& Bindings)
+inline FArchive& FMetalShaderBindings::Serialize(FArchive& Ar, FShaderResourceTable& SRT)
 {
-	Ar << Bindings.PackedUniformBuffers;
-	Ar << Bindings.PackedGlobalArrays;
-	Ar << Bindings.ShaderResourceTable;
-	Ar << Bindings.ArgumentBufferMasks;
-	Ar << Bindings.ConstantBuffers;
-	Ar << Bindings.ArgumentBuffers;
-	Ar << Bindings.InOutMask;
-	Ar << Bindings.NumSamplers;
-	Ar << Bindings.NumUniformBuffers;
-	Ar << Bindings.NumUAVs;
-	Ar << Bindings.bDiscards;
-    Ar << Bindings.IRConverterReflectionJSON;
-    Ar << Bindings.RSNumCBVs;
-    Ar << Bindings.OutputSizeVS;
-    Ar << Bindings.MaxInputPrimitivesPerMeshThreadgroupGS;
+	Ar << PackedUniformBuffers;
+	Ar << PackedGlobalArrays;
+	Ar << SRT;
+	Ar << ArgumentBufferMasks;
+	Ar << ConstantBuffers;
+	Ar << ArgumentBuffers;
+	Ar << InOutMask;
+	Ar << NumSamplers;
+	Ar << NumUniformBuffers;
+	Ar << NumUAVs;
+	Ar << bDiscards;
+    Ar << IRConverterReflectionJSON;
+    Ar << RSNumCBVs;
+    Ar << OutputSizeVS;
+    Ar << MaxInputPrimitivesPerMeshThreadgroupGS;
 	return Ar;
 }
 
@@ -261,45 +249,48 @@ struct FMetalCodeHeader
 	, bDeviceFunctionConstants(false)
 	{
 	}
+
+	inline FArchive& Serialize(FArchive& Ar, FShaderResourceTable& SRT);
 };
 
-inline FArchive& operator<<(FArchive& Ar, FMetalCodeHeader& Header)
+
+inline FArchive& FMetalCodeHeader::Serialize(FArchive& Ar, FShaderResourceTable& SRT)
 {
-	Ar << Header.Bindings;
+	Bindings.Serialize(Ar, SRT);
 
-	int32 NumInfos = Header.UniformBuffersCopyInfo.Num();
+	int32 NumInfos = UniformBuffersCopyInfo.Num();
 	Ar << NumInfos;
 	if (Ar.IsSaving())
 	{
 		for (int32 Index = 0; Index < NumInfos; ++Index)
 		{
-			Ar << Header.UniformBuffersCopyInfo[Index];
+			Ar << UniformBuffersCopyInfo[Index];
 		}
 	}
 	else if (Ar.IsLoading())
 	{
-		Header.UniformBuffersCopyInfo.Empty(NumInfos);
+		UniformBuffersCopyInfo.Empty(NumInfos);
 		for (int32 Index = 0; Index < NumInfos; ++Index)
 		{
 			CrossCompiler::FUniformBufferCopyInfo Info;
 			Ar << Info;
-			Header.UniformBuffersCopyInfo.Add(Info);
+			UniformBuffersCopyInfo.Add(Info);
 		}
 	}
 	
-	Ar << Header.CompilerBuild;
-	Ar << Header.CompilerVersion;
-	Ar << Header.SourceLen;
-	Ar << Header.SourceCRC;
-	Ar << Header.NumThreadsX;
-	Ar << Header.NumThreadsY;
-	Ar << Header.NumThreadsZ;
-	Ar << Header.CompileFlags;
-	Ar << Header.Frequency;
-	Ar << Header.Version;
-	Ar << Header.SideTable;
-	Ar << Header.bDeviceFunctionConstants;
-	Ar << Header.RayTracing;
+	Ar << CompilerBuild;
+	Ar << CompilerVersion;
+	Ar << SourceLen;
+	Ar << SourceCRC;
+	Ar << NumThreadsX;
+	Ar << NumThreadsY;
+	Ar << NumThreadsZ;
+	Ar << CompileFlags;
+	Ar << Frequency;
+	Ar << Version;
+	Ar << SideTable;
+	Ar << bDeviceFunctionConstants;
+	Ar << RayTracing;
     return Ar;
 }
 
diff --git a/Engine/Source/Runtime/D3D12RHI/Private/D3D12Commands.cpp b/Engine/Source/Runtime/D3D12RHI/Private/D3D12Commands.cpp
--- a/Engine/Source/Runtime/D3D12RHI/Private/D3D12Commands.cpp
+++ b/Engine/Source/Runtime/D3D12RHI/Private/D3D12Commands.cpp
@@ -49,21 +49,20 @@ inline FD3D12ShaderData* GetShaderData(FRHIShader* InShaderRHI)
 	return nullptr;
 }
 
-inline void ValidateBoundUniformBuffer(FD3D12UniformBuffer * InUniformBuffer, FRHIShader* InShaderRHI, uint32 InBufferIndex)
+inline void ValidateBoundUniformBuffer(FD3D12UniformBuffer* InUniformBuffer, FRHIShader* InShaderRHI, uint32 InBufferIndex)
 {
 #if DO_CHECK
-	if (FD3D12ShaderData* ShaderData = GetShaderData(InShaderRHI))
+	auto const& LayoutHashes = InShaderRHI->GetShaderResourceTable().ResourceTableLayoutHashes;
+
+	if (InBufferIndex < (uint32)LayoutHashes.Num())
 	{
-		if (InBufferIndex < (uint32)ShaderData->ShaderResourceTable.ResourceTableLayoutHashes.Num())
-		{
-			uint32 UniformBufferHash = InUniformBuffer->GetLayout().GetHash();
-			uint32 ShaderTableHash = ShaderData->ShaderResourceTable.ResourceTableLayoutHashes[InBufferIndex];
-			ensureMsgf(ShaderTableHash == 0 || UniformBufferHash == ShaderTableHash,
-				TEXT("Invalid uniform buffer %s bound on %sShader at index %d."),
-				*(InUniformBuffer->GetLayout().GetDebugName()),
-				GetShaderFrequencyString(InShaderRHI->GetFrequency(), false),
-				InBufferIndex);
-		}
+		uint32 UniformBufferHash = InUniformBuffer->GetLayout().GetHash();
+		uint32 ShaderTableHash = LayoutHashes[InBufferIndex];
+		ensureMsgf(ShaderTableHash == 0 || UniformBufferHash == ShaderTableHash,
+			TEXT("Invalid uniform buffer %s bound on %sShader at index %d."),
+			*(InUniformBuffer->GetLayout().GetDebugName()),
+			GetShaderFrequencyString(InShaderRHI->GetFrequency(), false),
+			InBufferIndex);
 	}
 #endif
 }
@@ -804,7 +803,7 @@ void FD3D12CommandContext::RHISetScissorRect(bool bEnable, uint32 MinX, uint32 M
 	}
 }
 
-static void ApplyStaticUniformBuffersOnContext(FD3D12CommandContext& Context, FRHIShader* Shader, FD3D12ShaderData* ShaderData)
+static void ApplyStaticUniformBuffersOnContext(FD3D12CommandContext& Context, FRHIShader* Shader)
 {
 	if (Shader)
 	{
@@ -813,8 +812,6 @@ static void ApplyStaticUniformBuffersOnContext(FD3D12CommandContext& Context, FR
 
 		UE::RHICore::ApplyStaticUniformBuffers(
 			Shader,
-			ShaderData->StaticSlots,
-			ShaderData->ShaderResourceTable.ResourceTableLayoutHashes,
 			Context.GetStaticUniformBuffers(),
 			[&Context, Shader, ShaderFrequency, GpuIndex](int32 BufferIndex, FRHIUniformBuffer* Buffer)
 			{
@@ -824,12 +821,6 @@ static void ApplyStaticUniformBuffersOnContext(FD3D12CommandContext& Context, FR
 	}
 }
 
-template<typename TShader>
-static void ApplyStaticUniformBuffersOnContext(FD3D12CommandContext& Context, TShader* Shader)
-{
-	ApplyStaticUniformBuffersOnContext(Context, Shader, static_cast<FD3D12ShaderData*>(Shader));
-}
-
 void FD3D12CommandContext::RHISetGraphicsPipelineState(FRHIGraphicsPipelineState* GraphicsState, uint32 StencilRef, bool bApplyAdditionalState)
 {
 	FD3D12GraphicsPipelineState* GraphicsPipelineState = FD3D12DynamicRHI::ResourceCast(GraphicsState);
@@ -1479,7 +1470,6 @@ void FD3D12CommandContext::SetResourcesFromTables(const ShaderType* Shader)
 	UE::RHICore::SetResourcesFromTables(
 		  Binder
 		, *Shader
-		, Shader->ShaderResourceTable
 		, DirtyUniformBuffers[Frequency]
 		, BoundUniformBuffers[Frequency]
 #if ENABLE_RHI_VALIDATION
diff --git a/Engine/Source/Runtime/D3D12RHI/Private/D3D12RayTracing.cpp b/Engine/Source/Runtime/D3D12RHI/Private/D3D12RayTracing.cpp
--- a/Engine/Source/Runtime/D3D12RHI/Private/D3D12RayTracing.cpp
+++ b/Engine/Source/Runtime/D3D12RHI/Private/D3D12RayTracing.cpp
@@ -4306,7 +4306,6 @@ static bool SetRayTracingShaderResources(
 		UE::RHICore::SetResourcesFromTables(
 			  Bindings
 			, *Shader
-			, Shader->ShaderResourceTable
 			, DirtyUniformBuffers
 			, UniformBuffers
 #if ENABLE_RHI_VALIDATION
diff --git a/Engine/Source/Runtime/D3D12RHI/Private/D3D12Shader.h b/Engine/Source/Runtime/D3D12RHI/Private/D3D12Shader.h
--- a/Engine/Source/Runtime/D3D12RHI/Private/D3D12Shader.h
+++ b/Engine/Source/Runtime/D3D12RHI/Private/D3D12Shader.h
@@ -70,13 +70,8 @@ struct FD3D12ShaderData
 	/** The shader's bytecode, with custom data in the last byte. */
 	TArray<uint8> Code;
 
-	FShaderResourceTable ShaderResourceTable;
-
 	FShaderCodePackedResourceCounts ResourceCounts{};
 
-	/** The static slot associated with the resource table index in ShaderResourceTable. */
-	TArray<FUniformBufferStaticSlot> StaticSlots;
-
 #if D3D12RHI_NEEDS_VENDOR_EXTENSIONS
 	TArray<FShaderCodeVendorExtension> VendorExtensions;
 #endif
diff --git a/Engine/Source/Runtime/D3D12RHI/Private/D3D12Shaders.cpp b/Engine/Source/Runtime/D3D12RHI/Private/D3D12Shaders.cpp
--- a/Engine/Source/Runtime/D3D12RHI/Private/D3D12Shaders.cpp
+++ b/Engine/Source/Runtime/D3D12RHI/Private/D3D12Shaders.cpp
@@ -122,7 +122,7 @@ TShaderType* CreateStandardShader(TArrayView<const uint8> InCode)
 	TShaderType* Shader = new TShaderType();
 
 	FMemoryReaderView Ar(InCode, true);
-	Ar << Shader->ShaderResourceTable;
+	Shader->SerializeShaderResourceTable(Ar);
 
 	const int32 Offset = Ar.Tell();
 
@@ -133,7 +133,8 @@ TShaderType* CreateStandardShader(TArrayView<const uint8> InCode)
 		return nullptr;
 	}
 
-	UE::RHICore::InitStaticUniformBufferSlots(Shader->StaticSlots, Shader->ShaderResourceTable);
+	UE::RHICore::InitStaticUniformBufferSlots(Shader);
+
 	return Shader;
 }
 
@@ -184,7 +185,7 @@ FRayTracingShaderRHIRef FD3D12DynamicRHI::RHICreateRayTracingShader(TArrayView<c
 	FD3D12RayTracingShader* Shader = new FD3D12RayTracingShader(ShaderFrequency);
 
 	FMemoryReaderView Ar(Code, true);
-	Ar << Shader->ShaderResourceTable;
+	Shader->SerializeShaderResourceTable(Ar);
 	Ar << Shader->EntryPoint;
 	Ar << Shader->AnyHitEntryPoint;
 	Ar << Shader->IntersectionEntryPoint;
@@ -216,7 +217,7 @@ FRayTracingShaderRHIRef FD3D12DynamicRHI::RHICreateRayTracingShader(TArrayView<c
 		return nullptr;
 	}
 
-	UE::RHICore::InitStaticUniformBufferSlots(Shader->StaticSlots, Shader->ShaderResourceTable);
+	UE::RHICore::InitStaticUniformBufferSlots(Shader);
 
 	Shader->pRootSignature = GetAdapter().GetRootSignature(Shader);
 
diff --git a/Engine/Source/Runtime/RHI/Public/RHIFwd.h b/Engine/Source/Runtime/RHI/Public/RHIFwd.h
--- a/Engine/Source/Runtime/RHI/Public/RHIFwd.h
+++ b/Engine/Source/Runtime/RHI/Public/RHIFwd.h
@@ -71,6 +71,7 @@ class FRHIRenderQueryPool;
 class FRHIResource;
 class FRHISamplerState;
 class FRHIShader;
+class FRHIShaderData;
 class FRHIShaderLibrary;
 class FRHIShaderResourceView;
 class FRHIShaderBundle;
diff --git a/Engine/Source/Runtime/RHI/Public/RHIResources.h b/Engine/Source/Runtime/RHI/Public/RHIResources.h
--- a/Engine/Source/Runtime/RHI/Public/RHIResources.h
+++ b/Engine/Source/Runtime/RHI/Public/RHIResources.h
@@ -758,7 +758,100 @@ public:
 	#define RHI_IF_SHADER_DEBUG_DATA(...)
 #endif
 
-class FRHIShader : public FRHIResource
+struct FShaderResourceTable
+{
+	/** Bits indicating which resource tables contain resources bound to this shader. */
+	uint32 ResourceTableBits = 0;
+
+	/** Mapping of bound SRVs to their location in resource tables. */
+	TArray<uint32> ShaderResourceViewMap;
+
+	/** Mapping of bound sampler states to their location in resource tables. */
+	TArray<uint32> SamplerMap;
+
+	/** Mapping of bound UAVs to their location in resource tables. */
+	TArray<uint32> UnorderedAccessViewMap;
+
+	/** Hash of the layouts of resource tables at compile time, used for runtime validation. */
+	TArray<uint32> ResourceTableLayoutHashes;
+
+	/** Mapping of bound Textures to their location in resource tables. */
+	TArray<uint32> TextureMap;
+
+	friend bool operator == (const FShaderResourceTable& A, const FShaderResourceTable& B)
+	{
+		bool bEqual = true;
+		bEqual &= (A.ResourceTableBits == B.ResourceTableBits);
+		bEqual &= (A.ShaderResourceViewMap    .Num() == B.ShaderResourceViewMap    .Num());
+		bEqual &= (A.SamplerMap               .Num() == B.SamplerMap               .Num());
+		bEqual &= (A.UnorderedAccessViewMap   .Num() == B.UnorderedAccessViewMap   .Num());
+		bEqual &= (A.ResourceTableLayoutHashes.Num() == B.ResourceTableLayoutHashes.Num());
+		bEqual &= (A.TextureMap               .Num() == B.TextureMap               .Num());
+
+		if (!bEqual)
+		{
+			return false;
+		}
+
+		bEqual &= (FMemory::Memcmp(A.ShaderResourceViewMap    .GetData(), B.ShaderResourceViewMap    .GetData(), A.ShaderResourceViewMap    .GetTypeSize() * A.ShaderResourceViewMap    .Num()) == 0);
+		bEqual &= (FMemory::Memcmp(A.SamplerMap               .GetData(), B.SamplerMap               .GetData(), A.SamplerMap               .GetTypeSize() * A.SamplerMap               .Num()) == 0);
+		bEqual &= (FMemory::Memcmp(A.UnorderedAccessViewMap   .GetData(), B.UnorderedAccessViewMap   .GetData(), A.UnorderedAccessViewMap   .GetTypeSize() * A.UnorderedAccessViewMap   .Num()) == 0);
+		bEqual &= (FMemory::Memcmp(A.ResourceTableLayoutHashes.GetData(), B.ResourceTableLayoutHashes.GetData(), A.ResourceTableLayoutHashes.GetTypeSize() * A.ResourceTableLayoutHashes.Num()) == 0);
+		bEqual &= (FMemory::Memcmp(A.TextureMap               .GetData(), B.TextureMap               .GetData(), A.TextureMap               .GetTypeSize() * A.TextureMap               .Num()) == 0);
+		return bEqual;
+	}
+
+	friend FArchive& operator << (FArchive& Ar, FShaderResourceTable& SRT)
+	{
+		Ar << SRT.ResourceTableBits;
+		Ar << SRT.ShaderResourceViewMap;
+		Ar << SRT.SamplerMap;
+		Ar << SRT.UnorderedAccessViewMap;
+		Ar << SRT.ResourceTableLayoutHashes;
+		Ar << SRT.TextureMap;
+	
+		return Ar;
+	}
+};
+
+namespace UE
+{
+	namespace RHICore
+	{
+		// Workaround for layering issue. FShaderParametersMetadata is in RenderCore,
+		// so we can't move the logic for initializing the StaticSlots array to the RHI module.
+		void InitStaticUniformBufferSlots(FRHIShaderData* ShaderData);
+	}
+}
+
+class FRHIShaderData
+{
+public:
+	const FShaderResourceTable& GetShaderResourceTable() const
+	{
+		return ShaderResourceTable;
+	}
+
+	const TArray<FUniformBufferStaticSlot>& GetStaticSlots() const
+	{
+		return StaticSlots;
+	}
+
+	void SerializeShaderResourceTable(FArchive& Ar)
+	{
+		Ar << ShaderResourceTable;
+	}
+
+protected:
+	FShaderResourceTable ShaderResourceTable;
+	TArray<FUniformBufferStaticSlot> StaticSlots;
+
+	// Workaround for layering issue. FShaderParametersMetadata is in RenderCore,
+	// so we can't move the logic for initializing the StaticSlots array to the RHI module.
+	friend void UE::RHICore::InitStaticUniformBufferSlots(FRHIShaderData* ShaderData);
+};
+
+class FRHIShader : public FRHIResource, public FRHIShaderData
 {
 public:
 	void SetHash(FSHAHash InHash) { Hash = InHash; }
diff --git a/Engine/Source/Runtime/RHICore/Public/RHICoreShader.h b/Engine/Source/Runtime/RHICore/Public/RHICoreShader.h
--- a/Engine/Source/Runtime/RHICore/Public/RHICoreShader.h
+++ b/Engine/Source/Runtime/RHICore/Public/RHICoreShader.h
@@ -29,8 +29,11 @@ extern RHICORE_API void DispatchShaderBundleEmulation(
 	TConstArrayView<FRHIShaderBundleDispatch> Dispatches
 );
 
-inline void InitStaticUniformBufferSlots(TArray<FUniformBufferStaticSlot>& StaticSlots, const FShaderResourceTable& ShaderResourceTable)
+inline void InitStaticUniformBufferSlots(FRHIShaderData* ShaderData)
 {
+	TArray<FUniformBufferStaticSlot>& StaticSlots = ShaderData->StaticSlots;
+	const FShaderResourceTable& ShaderResourceTable = ShaderData->GetShaderResourceTable();
+
 	StaticSlots.Reserve(ShaderResourceTable.ResourceTableLayoutHashes.Num());
 
 	for (uint32 LayoutHash : ShaderResourceTable.ResourceTableLayoutHashes)
@@ -49,12 +52,13 @@ inline void InitStaticUniformBufferSlots(TArray<FUniformBufferStaticSlot>& Stati
 template <typename TApplyFunction>
 void ApplyStaticUniformBuffers(
 	FRHIShader* Shader,
-	const TArray<FUniformBufferStaticSlot>& Slots,
-	const TArray<uint32>& LayoutHashes,
 	const TArray<FRHIUniformBuffer*>& UniformBuffers,
 	TApplyFunction&& ApplyFunction
 )
 {
+	const TArray<uint32>& LayoutHashes = Shader->GetShaderResourceTable().ResourceTableLayoutHashes;
+	const TArray<FUniformBufferStaticSlot>& Slots = Shader->GetStaticSlots();
+
 	checkf(LayoutHashes.Num() == Slots.Num(), TEXT("Shader %s, LayoutHashes %d, Slots %d"),
 		Shader->GetShaderName(), LayoutHashes.Num(), Slots.Num());
 
@@ -79,11 +83,9 @@ template <typename TRHIContext, typename TRHIShader>
 void ApplyStaticUniformBuffers(
 	TRHIContext* CommandContext,
 	TRHIShader* Shader,
-	const TArray<FUniformBufferStaticSlot>& Slots,
-	const TArray<uint32>& LayoutHashes,
 	const TArray<FRHIUniformBuffer*>& UniformBuffers)
 {
-	ApplyStaticUniformBuffers(Shader, Slots, LayoutHashes, UniformBuffers,
+	ApplyStaticUniformBuffers(Shader, UniformBuffers,
 		[CommandContext, Shader](int32 BufferIndex, FRHIUniformBuffer* Buffer)
 		{
 			CommandContext->RHISetShaderUniformBuffer(Shader, BufferIndex, Buffer);
@@ -130,7 +132,7 @@ inline void EnumerateUniformBufferResources(FRHIUniformBuffer* RESTRICT Buffer,
 }
 
 template <typename TBinder, typename TUniformBufferArrayType, typename TBitMaskType>
-void SetResourcesFromTables(TBinder&& Binder, FRHIShader const& Shader, FShaderResourceTable const& SRT, TBitMaskType& DirtyUniformBuffers, TUniformBufferArrayType const& BoundUniformBuffers
+void SetResourcesFromTables(TBinder&& Binder, FRHIShader const& Shader, TBitMaskType& DirtyUniformBuffers, TUniformBufferArrayType const& BoundUniformBuffers
 #if ENABLE_RHI_VALIDATION
 	, RHIValidation::FTracker* Tracker
 #endif
@@ -138,6 +140,8 @@ void SetResourcesFromTables(TBinder&& Binder, FRHIShader const& Shader, FShaderR
 {
 	float CurrentTimeForTextureTimes = FApp::GetCurrentTime();
 
+	FShaderResourceTable const& SRT = Shader.GetShaderResourceTable();
+
 	// Mask the dirty bits by those buffers from which the shader has bound resources.
 	uint32 DirtyBits = SRT.ResourceTableBits & DirtyUniformBuffers;
 	while (DirtyBits)
diff --git a/Engine/Source/Runtime/RenderCore/Public/ShaderCore.h b/Engine/Source/Runtime/RenderCore/Public/ShaderCore.h
--- a/Engine/Source/Runtime/RenderCore/Public/ShaderCore.h
+++ b/Engine/Source/Runtime/RenderCore/Public/ShaderCore.h
@@ -327,62 +327,6 @@ public:
 	TMap<FString,FParameterAllocation> ParameterMap;
 };
 
-struct FShaderResourceTable
-{
-	/** Bits indicating which resource tables contain resources bound to this shader. */
-	uint32 ResourceTableBits = 0;
-
-	/** Mapping of bound SRVs to their location in resource tables. */
-	TArray<uint32> ShaderResourceViewMap;
-
-	/** Mapping of bound sampler states to their location in resource tables. */
-	TArray<uint32> SamplerMap;
-
-	/** Mapping of bound UAVs to their location in resource tables. */
-	TArray<uint32> UnorderedAccessViewMap;
-
-	/** Hash of the layouts of resource tables at compile time, used for runtime validation. */
-	TArray<uint32> ResourceTableLayoutHashes;
-
-	/** Mapping of bound Textures to their location in resource tables. */
-	TArray<uint32> TextureMap;
-
-	friend bool operator==(const FShaderResourceTable&A, const FShaderResourceTable& B)
-	{
-		bool bEqual = true;
-		bEqual &= (A.ResourceTableBits == B.ResourceTableBits);
-		bEqual &= (A.ShaderResourceViewMap    .Num() == B.ShaderResourceViewMap    .Num());
-		bEqual &= (A.SamplerMap               .Num() == B.SamplerMap               .Num());
-		bEqual &= (A.UnorderedAccessViewMap   .Num() == B.UnorderedAccessViewMap   .Num());
-		bEqual &= (A.ResourceTableLayoutHashes.Num() == B.ResourceTableLayoutHashes.Num());
-		bEqual &= (A.TextureMap               .Num() == B.TextureMap               .Num());
-
-		if (!bEqual)
-		{
-			return false;
-		}
-
-		bEqual &= (FMemory::Memcmp(A.ShaderResourceViewMap    .GetData(), B.ShaderResourceViewMap    .GetData(), A.ShaderResourceViewMap    .GetTypeSize() * A.ShaderResourceViewMap    .Num()) == 0);
-		bEqual &= (FMemory::Memcmp(A.SamplerMap               .GetData(), B.SamplerMap               .GetData(), A.SamplerMap               .GetTypeSize() * A.SamplerMap               .Num()) == 0);
-		bEqual &= (FMemory::Memcmp(A.UnorderedAccessViewMap   .GetData(), B.UnorderedAccessViewMap   .GetData(), A.UnorderedAccessViewMap   .GetTypeSize() * A.UnorderedAccessViewMap   .Num()) == 0);
-		bEqual &= (FMemory::Memcmp(A.ResourceTableLayoutHashes.GetData(), B.ResourceTableLayoutHashes.GetData(), A.ResourceTableLayoutHashes.GetTypeSize() * A.ResourceTableLayoutHashes.Num()) == 0);
-		bEqual &= (FMemory::Memcmp(A.TextureMap               .GetData(), B.TextureMap               .GetData(), A.TextureMap               .GetTypeSize() * A.TextureMap               .Num()) == 0);
-		return bEqual;
-	}
-};
-
-inline FArchive& operator<<(FArchive& Ar, FShaderResourceTable& SRT)
-{
-	Ar << SRT.ResourceTableBits;
-	Ar << SRT.ShaderResourceViewMap;
-	Ar << SRT.SamplerMap;
-	Ar << SRT.UnorderedAccessViewMap;
-	Ar << SRT.ResourceTableLayoutHashes;
-	Ar << SRT.TextureMap;
-
-	return Ar;
-}
-
 inline FArchive& operator<<(FArchive& Ar, FUniformResourceEntry& Entry)
 {
 	if (Ar.IsLoading())
diff --git a/Engine/Source/Runtime/VulkanRHI/Private/VulkanCommands.cpp b/Engine/Source/Runtime/VulkanRHI/Private/VulkanCommands.cpp
--- a/Engine/Source/Runtime/VulkanRHI/Private/VulkanCommands.cpp
+++ b/Engine/Source/Runtime/VulkanRHI/Private/VulkanCommands.cpp
@@ -180,7 +180,6 @@ void FVulkanCommandListContext::SetResourcesFromTables(const ShaderType* Shader)
 		UE::RHICore::SetResourcesFromTables(
 			Binder
 			, *Shader
-			, Shader->ShaderResourceTable
 			, DirtyUniformBuffers[Frequency]
 			, BoundUniformBuffers[Frequency]
 #if ENABLE_RHI_VALIDATION
@@ -197,7 +196,6 @@ void FVulkanCommandListContext::SetResourcesFromTables(const ShaderType* Shader)
 		UE::RHICore::SetResourcesFromTables(
 			Binder
 			, *Shader
-			, Shader->ShaderResourceTable
 			, DirtyUniformBuffers[Frequency]
 			, BoundUniformBuffers[Frequency]
 #if ENABLE_RHI_VALIDATION
diff --git a/Engine/Source/Runtime/VulkanRHI/Private/VulkanPipelineState.cpp b/Engine/Source/Runtime/VulkanRHI/Private/VulkanPipelineState.cpp
--- a/Engine/Source/Runtime/VulkanRHI/Private/VulkanPipelineState.cpp
+++ b/Engine/Source/Runtime/VulkanRHI/Private/VulkanPipelineState.cpp
@@ -535,7 +535,7 @@ void FVulkanCommandListContext::ApplyStaticUniformBuffers(TRHIShader* Shader)
 {
 	if (Shader)
 	{
-		const auto& StaticSlots = Shader->StaticSlots;
+		const auto& StaticSlots = Shader->GetStaticSlots();
 		const auto& UBInfos = Shader->GetCodeHeader().UniformBuffers;
 
 		for (int32 BufferIndex = 0; BufferIndex < StaticSlots.Num(); ++BufferIndex)
diff --git a/Engine/Source/Runtime/VulkanRHI/Private/VulkanResources.h b/Engine/Source/Runtime/VulkanRHI/Private/VulkanResources.h
--- a/Engine/Source/Runtime/VulkanRHI/Private/VulkanResources.h
+++ b/Engine/Source/Runtime/VulkanRHI/Private/VulkanResources.h
@@ -105,13 +105,6 @@ protected:
 	static FCriticalSection VulkanShaderModulesMapCS;
 
 public:
-	FVulkanShader(FVulkanDevice* InDevice, EShaderFrequency InFrequency)
-		: ShaderKey(0)
-		, Frequency(InFrequency)
-		, Device(InDevice)
-	{
-	}
-
 	virtual ~FVulkanShader();
 
 	void PurgeShaderModules();
@@ -190,6 +183,9 @@ public:
 	}
 	
 	FSpirvCode GetPatchedSpirvCode(const FGfxPipelineDesc& Desc, const FVulkanLayout* Layout);
+
+	TArray<FUniformBufferStaticSlot>& StaticSlots;
+
 protected:
 
 #if UE_BUILD_DEBUG || UE_BUILD_DEVELOPMENT
@@ -202,10 +198,6 @@ protected:
 	TMap<uint32, TRefCountPtr<FVulkanShaderModule>>	ShaderModules;
 	const EShaderFrequency			Frequency;
 
-	TArray<FUniformBufferStaticSlot> StaticSlots;
-
-	FShaderResourceTable			ShaderResourceTable;
-
 protected:
 	class FSpirvContainer
 	{
@@ -224,8 +216,6 @@ protected:
 	static FSpirvCode GetSpirvCode(const FSpirvContainer& Container);
 
 protected:
-	void Setup(FVulkanShaderHeader&& InCodeHeader, FShaderResourceTable&& InSRT, FSpirvContainer&& InSpirvContainer, uint64 InShaderKey);
-
 	FVulkanDevice*					Device;
 
 	TRefCountPtr<FVulkanShaderModule> CreateHandle(const FVulkanLayout* Layout, uint32 LayoutHash);
@@ -233,6 +223,8 @@ protected:
 
 	bool NeedsSpirvInputAttachmentPatching(const FGfxPipelineDesc& Desc) const;
 
+	FVulkanShader(FVulkanDevice* InDevice, EShaderFrequency InFrequency, FVulkanShaderHeader&& InCodeHeader, FSpirvContainer&& InSpirvContainer, uint64 InShaderKey, TArray<FUniformBufferStaticSlot>& InStaticSlots);
+
 	friend class FVulkanCommandListContext;
 	friend class FVulkanPipelineStateCacheManager;
 	friend class FVulkanComputeShaderState;
@@ -245,9 +237,11 @@ template<typename BaseResourceType, EShaderFrequency ShaderType>
 class TVulkanBaseShader : public BaseResourceType, public FVulkanShader
 {
 private:
-	TVulkanBaseShader(FVulkanDevice* InDevice) :
-		FVulkanShader(InDevice, ShaderType)
+	TVulkanBaseShader(FVulkanDevice* InDevice, FShaderResourceTable&& InSRT, FVulkanShaderHeader&& InCodeHeader, FSpirvContainer&& InSpirvContainer, uint64 InShaderKey)
+		: BaseResourceType()
+		, FVulkanShader(InDevice, ShaderType, MoveTemp(InCodeHeader), MoveTemp(InSpirvContainer), InShaderKey, BaseResourceType::StaticSlots)
 	{
+		BaseResourceType::ShaderResourceTable = MoveTemp(InSRT);
 	}
 	friend class FVulkanShaderFactory;
 public:
@@ -277,10 +271,11 @@ typedef TVulkanBaseShader<FRHIGeometryShader, SF_Geometry>			FVulkanGeometryShad
 class FVulkanRayTracingShader : public FRHIRayTracingShader, public FVulkanShader
 {
 private:
-	FVulkanRayTracingShader(FVulkanDevice* InDevice, EShaderFrequency InFrequency)
+	FVulkanRayTracingShader(FVulkanDevice* InDevice, EShaderFrequency InFrequency, FShaderResourceTable&& InSRT, FVulkanShaderHeader&& InCodeHeader, FSpirvContainer&& InSpirvContainer, uint64 InShaderKey)
 		: FRHIRayTracingShader(InFrequency)
-		, FVulkanShader(InDevice, InFrequency)
+		, FVulkanShader(InDevice, InFrequency, MoveTemp(InCodeHeader), MoveTemp(InSpirvContainer), InShaderKey, FRHIRayTracingShader::StaticSlots)
 	{
+		ShaderResourceTable = MoveTemp(InSRT);
 	}
 
 	FSpirvContainer AnyHitSpirvContainer;
diff --git a/Engine/Source/Runtime/VulkanRHI/Private/VulkanShaders.cpp b/Engine/Source/Runtime/VulkanRHI/Private/VulkanShaders.cpp
--- a/Engine/Source/Runtime/VulkanRHI/Private/VulkanShaders.cpp
+++ b/Engine/Source/Runtime/VulkanRHI/Private/VulkanShaders.cpp
@@ -81,8 +81,8 @@ ShaderType* FVulkanShaderFactory::CreateShader(TArrayView<const uint8> Code, FVu
 			}
 			else
 			{
-				RetShader = new ShaderType(Device);
-				RetShader->Setup(MoveTemp(CodeHeader), MoveTemp(SerializedSRT), MoveTemp(SpirvContainer), ShaderKey);
+				RetShader = new ShaderType(Device, MoveTemp(SerializedSRT), MoveTemp(CodeHeader), MoveTemp(SpirvContainer), ShaderKey);
+
 				ShaderMap[ShaderType::StaticFrequency].Add(ShaderKey, RetShader);
 
 #if !(UE_BUILD_SHIPPING || UE_BUILD_TEST)
@@ -159,8 +159,8 @@ FVulkanRayTracingShader* FVulkanShaderFactory::CreateRayTracingShader(TArrayView
 			}
 			else
 			{
-				RetShader = new FVulkanRayTracingShader(Device, ShaderFrequency);
-				RetShader->Setup(MoveTemp(CodeHeader), MoveTemp(SerializedSRT), MoveTemp(SpirvContainer), ShaderKey);
+				RetShader = new FVulkanRayTracingShader(Device, ShaderFrequency, MoveTemp(SerializedSRT), MoveTemp(CodeHeader), MoveTemp(SpirvContainer), ShaderKey);
+
 				if (bIsHitGroup)
 				{
 					RetShader->AnyHitSpirvContainer = MoveTemp(AnyHitSpirvContainer);
@@ -275,19 +275,17 @@ FVulkanShader::FSpirvCode FVulkanShader::GetSpirvCode(const FSpirvContainer& Con
 }
 
 
-void FVulkanShader::Setup(FVulkanShaderHeader&& InCodeHeader, FShaderResourceTable&& InSRT, FSpirvContainer&& InSpirvContainer, uint64 InShaderKey)
+FVulkanShader::FVulkanShader(FVulkanDevice* InDevice, EShaderFrequency InFrequency, FVulkanShaderHeader&& InCodeHeader, FSpirvContainer&& InSpirvContainer, uint64 InShaderKey, TArray<FUniformBufferStaticSlot>& InStaticSlots)
+	: StaticSlots(InStaticSlots)
+	, ShaderKey(InShaderKey)
+	, CodeHeader(MoveTemp(InCodeHeader))
+	, Frequency(InFrequency)
+	, SpirvContainer(MoveTemp(InSpirvContainer))
+	, Device(InDevice)
 {
 	LLM_SCOPE_VULKAN(ELLMTagVulkan::VulkanShaders);
 	check(Device);
 
-	ShaderKey = InShaderKey;
-
-	CodeHeader = MoveTemp(InCodeHeader);
-
-	ShaderResourceTable = MoveTemp(InSRT);
-
-	SpirvContainer = MoveTemp(InSpirvContainer);
-
 	checkf(SpirvContainer.GetSizeBytes() != 0, TEXT("Empty SPIR-V! %s"), *CodeHeader.DebugName);
 
 	check(IsRayTracingShaderFrequency(Frequency) || (CodeHeader.UniformBufferSpirvInfos.Num() == CodeHeader.UniformBuffers.Num()));
diff --git a/Engine/Source/Runtime/Windows/D3D11RHI/Private/D3D11Commands.cpp b/Engine/Source/Runtime/Windows/D3D11RHI/Private/D3D11Commands.cpp
--- a/Engine/Source/Runtime/Windows/D3D11RHI/Private/D3D11Commands.cpp
+++ b/Engine/Source/Runtime/Windows/D3D11RHI/Private/D3D11Commands.cpp
@@ -120,7 +120,7 @@ void FD3D11DynamicRHI::ApplyStaticUniformBuffers(TRHIShader* Shader)
 {
 	if (Shader)
 	{
-		UE::RHICore::ApplyStaticUniformBuffers(Shader, Shader->StaticSlots, Shader->ShaderResourceTable.ResourceTableLayoutHashes, StaticUniformBuffers,
+		UE::RHICore::ApplyStaticUniformBuffers(Shader, StaticUniformBuffers,
 			[this](int32 BufferIndex, FRHIUniformBuffer* Buffer)
 			{
 				BindUniformBuffer<static_cast<EShaderFrequency>(TRHIShader::StaticFrequency)>(BufferIndex, Buffer);
@@ -1008,7 +1008,6 @@ void FD3D11DynamicRHI::SetResourcesFromTables(const ShaderType* RESTRICT Shader)
 	UE::RHICore::SetResourcesFromTables(
 		  FD3D11ResourceBinder<Frequency> { *this }
 		, *Shader
-		, Shader->ShaderResourceTable
 		, DirtyUniformBuffers[Frequency]
 		, BoundUniformBuffers[Frequency]
 #if ENABLE_RHI_VALIDATION
diff --git a/Engine/Source/Runtime/Windows/D3D11RHI/Private/D3D11Shaders.cpp b/Engine/Source/Runtime/Windows/D3D11RHI/Private/D3D11Shaders.cpp
--- a/Engine/Source/Runtime/Windows/D3D11RHI/Private/D3D11Shaders.cpp
+++ b/Engine/Source/Runtime/Windows/D3D11RHI/Private/D3D11Shaders.cpp
@@ -125,7 +125,7 @@ FVertexShaderRHIRef FD3D11DynamicRHI::RHICreateVertexShader(TArrayView<const uin
 	FD3D11VertexShader* Shader = new FD3D11VertexShader;
 
 	FMemoryReaderView Ar( Code, true );
-	Ar << Shader->ShaderResourceTable;
+	Shader->SerializeShaderResourceTable(Ar);
 	int32 Offset = Ar.Tell();
 
 	TArrayView<const uint8> ActualCode = ShaderCode.GetOffsetShaderCode(Offset);
@@ -140,7 +140,7 @@ FVertexShaderRHIRef FD3D11DynamicRHI::RHICreateVertexShader(TArrayView<const uin
 		{
 			ResetVendorExtensions(Direct3DDevice);
 		}
-		UE::RHICore::InitStaticUniformBufferSlots(Shader->StaticSlots, Shader->ShaderResourceTable);
+		UE::RHICore::InitStaticUniformBufferSlots(Shader);
 	}
 	
 	// TEMP
@@ -156,7 +156,7 @@ FGeometryShaderRHIRef FD3D11DynamicRHI::RHICreateGeometryShader(TArrayView<const
 	FD3D11GeometryShader* Shader = new FD3D11GeometryShader;
 
 	FMemoryReaderView Ar( Code, true );
-	Ar << Shader->ShaderResourceTable;
+	Shader->SerializeShaderResourceTable(Ar);
 	int32 Offset = Ar.Tell();
 
 	TArrayView<const uint8> ActualCode = ShaderCode.GetOffsetShaderCode(Offset);
@@ -171,7 +171,7 @@ FGeometryShaderRHIRef FD3D11DynamicRHI::RHICreateGeometryShader(TArrayView<const
 		{
 			ResetVendorExtensions(Direct3DDevice);
 		}
-		UE::RHICore::InitStaticUniformBufferSlots(Shader->StaticSlots, Shader->ShaderResourceTable);
+		UE::RHICore::InitStaticUniformBufferSlots(Shader);
 	}
 
 	return Shader;
@@ -183,7 +183,7 @@ FPixelShaderRHIRef FD3D11DynamicRHI::RHICreatePixelShader(TArrayView<const uint8
 	FD3D11PixelShader* Shader = new FD3D11PixelShader;
 
 	FMemoryReaderView Ar( Code, true );
-	Ar << Shader->ShaderResourceTable;
+	Shader->SerializeShaderResourceTable(Ar);
 	int32 Offset = Ar.Tell();
 
 	TArrayView<const uint8> ActualCode = ShaderCode.GetOffsetShaderCode(Offset);
@@ -198,7 +198,7 @@ FPixelShaderRHIRef FD3D11DynamicRHI::RHICreatePixelShader(TArrayView<const uint8
 		{
 			ResetVendorExtensions(Direct3DDevice);
 		}
-		UE::RHICore::InitStaticUniformBufferSlots(Shader->StaticSlots, Shader->ShaderResourceTable);
+		UE::RHICore::InitStaticUniformBufferSlots(Shader);
 	}
 
 	return Shader;
@@ -210,7 +210,7 @@ FComputeShaderRHIRef FD3D11DynamicRHI::RHICreateComputeShader(TArrayView<const u
 	FD3D11ComputeShader* Shader = new FD3D11ComputeShader;
 
 	FMemoryReaderView Ar( Code, true );
-	Ar << Shader->ShaderResourceTable;
+	Shader->SerializeShaderResourceTable(Ar);
 	int32 Offset = Ar.Tell();
 
 	TArrayView<const uint8> ActualCode = ShaderCode.GetOffsetShaderCode(Offset);
@@ -225,7 +225,7 @@ FComputeShaderRHIRef FD3D11DynamicRHI::RHICreateComputeShader(TArrayView<const u
 		{
 			ResetVendorExtensions(Direct3DDevice);
 		}
-		UE::RHICore::InitStaticUniformBufferSlots(Shader->StaticSlots, Shader->ShaderResourceTable);
+		UE::RHICore::InitStaticUniformBufferSlots(Shader);
 	}
 
 	return Shader;
diff --git a/Engine/Source/Runtime/Windows/D3D11RHI/Public/D3D11Resources.h b/Engine/Source/Runtime/Windows/D3D11RHI/Public/D3D11Resources.h
--- a/Engine/Source/Runtime/Windows/D3D11RHI/Public/D3D11Resources.h
+++ b/Engine/Source/Runtime/Windows/D3D11RHI/Public/D3D11Resources.h
@@ -40,8 +40,6 @@ public:
 
 struct FD3D11ShaderData
 {
-	FShaderResourceTable				ShaderResourceTable;
-	TArray<FUniformBufferStaticSlot>	StaticSlots;
 	TArray<FShaderCodeVendorExtension>	VendorExtensions;
 	bool								bShaderNeedsGlobalConstantBuffer;
 	bool								bIsSm6Shader;
