Read the  windows native toolchain from the right path

diff --git a/Engine/Source/Developer/Apple/MetalShaderFormat/Private/MetalShaderFormat.cpp b/Engine/Source/Developer/Apple/MetalShaderFormat/Private/MetalShaderFormat.cpp
--- a/Engine/Source/Developer/Apple/MetalShaderFormat/Private/MetalShaderFormat.cpp
+++ b/Engine/Source/Developer/Apple/MetalShaderFormat/Private/MetalShaderFormat.cpp
@@ -921,8 +921,8 @@ FMetalCompilerToolchain::EMetalToolchainStatus FMetalCompilerToolchain::DoWindow
 	}
 
 	// Look for the windows native toolchain
-	MetalFrontendBinaryCommand[AppleSDKMac] = ToolchainBase / TEXT("macos") / TEXT("bin") / MetalFrontendBinary;
-	MetalFrontendBinaryCommand[AppleSDKMobile] = ToolchainBase / TEXT("ios") / TEXT("bin") / MetalFrontendBinary;
+	MetalFrontendBinaryCommand[AppleSDKMac] = ToolchainBase / TEXT("metal") / TEXT("macos") / TEXT("bin") / MetalFrontendBinary;
+	MetalFrontendBinaryCommand[AppleSDKMobile] = ToolchainBase / TEXT("metal") / TEXT("ios") / TEXT("bin") / MetalFrontendBinary;
 
 	bool bUseLocalMetalToolchain = FPaths::FileExists(MetalFrontendBinaryCommand[AppleSDKMac]) && FPaths::FileExists(MetalFrontendBinaryCommand[AppleSDKMobile]);
 	if (!bUseLocalMetalToolchain)
@@ -934,14 +934,14 @@ FMetalCompilerToolchain::EMetalToolchainStatus FMetalCompilerToolchain::DoWindow
 		return EMetalToolchainStatus::ToolchainNotFound;
 	}
 
-	MetalArBinaryCommand[AppleSDKMac] = ToolchainBase / TEXT("macos") / TEXT("bin") / MetalArBinary;
-	MetalArBinaryCommand[AppleSDKMobile] = ToolchainBase / TEXT("ios") / TEXT("bin") / MetalArBinary;
+	MetalArBinaryCommand[AppleSDKMac] = ToolchainBase / TEXT("metal") / TEXT("macos") / TEXT("bin") / MetalArBinary;
+	MetalArBinaryCommand[AppleSDKMobile] = ToolchainBase / TEXT("metal") / TEXT("ios") / TEXT("bin") / MetalArBinary;
 
-	MetalLibBinaryCommand[AppleSDKMac] = ToolchainBase / TEXT("macos") / TEXT("bin") / MetalLibraryBinary;
-	MetalLibBinaryCommand[AppleSDKMobile] = ToolchainBase / TEXT("ios") / TEXT("bin") / MetalLibraryBinary;
+	MetalLibBinaryCommand[AppleSDKMac] = ToolchainBase / TEXT("metal") / TEXT("macos") / TEXT("bin") / MetalLibraryBinary;
+	MetalLibBinaryCommand[AppleSDKMobile] = ToolchainBase / TEXT("metal") / TEXT("ios") / TEXT("bin") / MetalLibraryBinary;
 
-	AirPackBinaryCommand[AppleSDKMac] = ToolchainBase / TEXT("macos") / TEXT("bin") / AirPackBinary;
-    AirPackBinaryCommand[AppleSDKMobile] = ToolchainBase / TEXT("ios") / TEXT("bin") / AirPackBinary;
+	AirPackBinaryCommand[AppleSDKMac] = ToolchainBase / TEXT("metal") / TEXT("macos") / TEXT("bin") / AirPackBinary;
+    AirPackBinaryCommand[AppleSDKMobile] = ToolchainBase / TEXT("metal") / TEXT("ios") / TEXT("bin") / AirPackBinary;
     
 	if (!FPaths::FileExists(MetalArBinaryCommand[AppleSDKMac]) ||
 		!FPaths::FileExists(MetalArBinaryCommand[AppleSDKMobile]) ||
