Autosdk MetalDeveloperTools support
On Win64 MetalShaderFormat module is always loaded

diff --git a/Engine/Source/Developer/Apple/MetalShaderFormat/Private/MetalShaderFormat.cpp b/Engine/Source/Developer/Apple/MetalShaderFormat/Private/MetalShaderFormat.cpp
--- a/Engine/Source/Developer/Apple/MetalShaderFormat/Private/MetalShaderFormat.cpp
+++ b/Engine/Source/Developer/Apple/MetalShaderFormat/Private/MetalShaderFormat.cpp
@@ -612,6 +612,7 @@ FString FMetalCompilerToolchain::MetalMacSDK(TEXT("macosx"));
 FString FMetalCompilerToolchain::MetalMobileSDK(TEXT("iphoneos"));
 
 FString FMetalCompilerToolchain::DefaultWindowsToolchainPath(TEXT("c:/Program Files/Metal Developer Tools"));
+FString FMetalCompilerToolchain::WindowsToolchainVersion(TEXT("4.1"));
 
 // Static methods
 
@@ -912,14 +913,25 @@ FMetalCompilerToolchain::EMetalToolchainStatus FMetalCompilerToolchain::DoWindow
 	int32 Result = 0;
 	
 	FString ToolchainBase;
-	GConfig->GetString(TEXT("/Script/IOSRuntimeSettings.IOSRuntimeSettings"), TEXT("WindowsMetalToolchainOverride"), ToolchainBase, GEngineIni);
 
-	const bool bUseOverride = (!ToolchainBase.IsEmpty() && FPaths::DirectoryExists(ToolchainBase));
-	if (!bUseOverride)
+	static const FString SDKRootEnvFar(TEXT("UE_SDKS_ROOT"));
+	FString SDKPath = FPlatformMisc::GetEnvironmentVariable(*SDKRootEnvFar);
+
+	if (SDKPath.Len() != 0)
 	{
-		ToolchainBase = DefaultWindowsToolchainPath;
+		FString HostPlatform(TEXT("HostWin64"));
+		ToolchainBase = FPaths::Combine(*SDKPath, *HostPlatform, TEXT("Win64"), TEXT("MetalDeveloperTools"), WindowsToolchainVersion);
 	}
+	else
+	{
+		GConfig->GetString(TEXT("/Script/IOSRuntimeSettings.IOSRuntimeSettings"), TEXT("WindowsMetalToolchainOverride"), ToolchainBase, GEngineIni);
 
+		const bool bUseOverride = (!ToolchainBase.IsEmpty() && FPaths::DirectoryExists(ToolchainBase));
+		if (!bUseOverride)
+		{
+			ToolchainBase = DefaultWindowsToolchainPath;
+		}
+	}
 	// Look for the windows native toolchain
 	MetalFrontendBinaryCommand[AppleSDKMac] = ToolchainBase / TEXT("metal") / TEXT("macos") / TEXT("bin") / MetalFrontendBinary;
 	MetalFrontendBinaryCommand[AppleSDKMobile] = ToolchainBase / TEXT("metal") / TEXT("ios") / TEXT("bin") / MetalFrontendBinary;
diff --git a/Engine/Source/Developer/Apple/MetalShaderFormat/Private/MetalShaderFormat.h b/Engine/Source/Developer/Apple/MetalShaderFormat/Private/MetalShaderFormat.h
--- a/Engine/Source/Developer/Apple/MetalShaderFormat/Private/MetalShaderFormat.h
+++ b/Engine/Source/Developer/Apple/MetalShaderFormat/Private/MetalShaderFormat.h
@@ -210,6 +210,8 @@ public:
 	static FString MetalMobileSDK;
 	// The default installation directory of the windows native metal compiler
 	static FString DefaultWindowsToolchainPath;
+	// The version of the windows native metal compiler
+	static FString WindowsToolchainVersion;
 
 private:
 	// Members
diff --git a/Engine/Source/Programs/UnrealBuildTool/Platform/IOS/UEBuildIOS.cs b/Engine/Source/Programs/UnrealBuildTool/Platform/IOS/UEBuildIOS.cs
--- a/Engine/Source/Programs/UnrealBuildTool/Platform/IOS/UEBuildIOS.cs
+++ b/Engine/Source/Programs/UnrealBuildTool/Platform/IOS/UEBuildIOS.cs
@@ -1033,7 +1033,7 @@ namespace UnrealBuildTool
 
 			if ((Target.Platform == UnrealTargetPlatform.Win64) || (Target.Platform == UnrealTargetPlatform.Mac))
 			{
-				bool bBuildShaderFormats = Target.bForceBuildShaderFormats;
+				bool bBuildShaderFormats = Target.bForceBuildShaderFormats || Target.Platform == UnrealTargetPlatform.Win64;
 				if (!Target.bBuildRequiresCookedData)
 				{
 					if (ModuleName == "Engine")
@@ -1084,7 +1084,7 @@ namespace UnrealBuildTool
 						}
 					}
 
-					if (bBuildShaderFormats && bIsPlatformAvailableForTargetWithSDK)
+					if (bBuildShaderFormats)
 					{
 						Rules.DynamicallyLoadedModuleNames.Add("MetalShaderFormat");
 					}
