Fix Metal freeze when cooking on Windows

diff --git a/Engine/Source/Developer/Apple/MetalShaderFormat/Private/MetalShaderFormat.cpp b/Engine/Source/Developer/Apple/MetalShaderFormat/Private/MetalShaderFormat.cpp
--- a/Engine/Source/Developer/Apple/MetalShaderFormat/Private/MetalShaderFormat.cpp
+++ b/Engine/Source/Developer/Apple/MetalShaderFormat/Private/MetalShaderFormat.cpp
@@ -1040,7 +1040,7 @@ bool FMetalCompilerToolchain::ExecMetalAr(EAppleSDKType SDK, const TCHAR* Script
 	bool bSuccess = ExecGenericCommand(TEXT("/bin/sh"), *Command, OutReturnCode, OutStdOut, OutStdErr);
 #else
 	FString Command = FString::Printf(TEXT("/C type \"%s\" | \"%s\" -M"), ScriptFile, *this->MetalArBinaryCommand[SDK]);
-	bool bSuccess = ExecGenericCommand(TEXT("cmd.exe"), *Command, OutReturnCode, OutStdOut, OutStdErr);
+	bool bSuccess = ExecGenericCommand(TEXT("cmd.exe"), *Command, OutReturnCode, OutStdOut, OutStdErr, true /*bIsConsoleApp*/);
 #endif
 	if (!bSuccess)
 	{
@@ -1050,9 +1050,10 @@ bool FMetalCompilerToolchain::ExecMetalAr(EAppleSDKType SDK, const TCHAR* Script
 	return bSuccess;
 }
 
-bool FMetalCompilerToolchain::ExecGenericCommand(const TCHAR* Command, const TCHAR* Params, int32* OutReturnCode, FString* OutStdOut, FString* OutStdErr) const
+bool FMetalCompilerToolchain::ExecGenericCommand(const TCHAR* Command, const TCHAR* Params, int32* OutReturnCode, FString* OutStdOut, FString* OutStdErr, bool bIsConsoleApp) const
 {
 #if PLATFORM_WINDOWS
+	if(bIsConsoleApp)
 	{
 		// Why do we have our own implementation here? Because metal.exe wants to create a console window. 
 		// So if we don't specify the options to CreateProc we end up with tons and tons of windows appearing and disappearing during a cook.
@@ -1083,10 +1084,12 @@ bool FMetalCompilerToolchain::ExecGenericCommand(const TCHAR* Command, const TCH
 
 		return RC == 0;
 	}
-#else
-	// Otherwise use the API
-	return FPlatformProcess::ExecProcess(Command, Params, OutReturnCode, OutStdOut, OutStdErr);
+	else
 #endif
+	{
+		// Otherwise use the API
+		return FPlatformProcess::ExecProcess(Command, Params, OutReturnCode, OutStdOut, OutStdErr);
+	}
 }
 
 bool FMetalCompilerToolchain::CompileMetalShader(FMetalShaderBytecodeJob& Job, FMetalShaderBytecode& Output) const
diff --git a/Engine/Source/Developer/Apple/MetalShaderFormat/Private/MetalShaderFormat.h b/Engine/Source/Developer/Apple/MetalShaderFormat/Private/MetalShaderFormat.h
--- a/Engine/Source/Developer/Apple/MetalShaderFormat/Private/MetalShaderFormat.h
+++ b/Engine/Source/Developer/Apple/MetalShaderFormat/Private/MetalShaderFormat.h
@@ -73,7 +73,7 @@ public:
 	bool CompileMetalShader(struct FMetalShaderBytecodeJob& Job, struct FMetalShaderBytecode& Output) const;
 	
 	// Executes 'Command' on the local machine
-	bool ExecGenericCommand(const TCHAR* Command, const TCHAR* Params, int32* OutReturnCode, FString* OutStdOut, FString* OutStdErr) const;
+	bool ExecGenericCommand(const TCHAR* Command, const TCHAR* Params, int32* OutReturnCode, FString* OutStdOut, FString* OutStdErr, bool bIsConsoleApp = false) const;
 	// Executes the metal frontend compiler for 'SDK' on the local or remote machine, depending on the current configuration
 	bool ExecMetalFrontend(EAppleSDKType SDK, const TCHAR* Parameters, int32* OutReturnCode, FString* OutStdOut, FString* OutStdErr) const;
 	// Executes metallib for 'SDK' on the local or remote machine, depending on configuration
