[shaders] strip unnecessary source file hashing from metal compiler. this was dead code and the result wasn't being used anywhere.

diff --git a/Engine/Source/Developer/Apple/MetalShaderFormat/Private/MetalCompileShaderMSC.cpp b/Engine/Source/Developer/Apple/MetalShaderFormat/Private/MetalCompileShaderMSC.cpp
--- a/Engine/Source/Developer/Apple/MetalShaderFormat/Private/MetalCompileShaderMSC.cpp
+++ b/Engine/Source/Developer/Apple/MetalShaderFormat/Private/MetalCompileShaderMSC.cpp
@@ -25,7 +25,6 @@
 extern void BuildMetalShaderOutput(
 	FShaderCompilerOutput& ShaderOutput,
 	const FShaderCompilerInput& ShaderInput,
-	FSHAHash const& GUIDHash,
 	const ANSICHAR* InShaderSource,
 	uint32 SourceLen,
 	uint32 SourceCRCLen,
@@ -429,7 +428,6 @@ void FMetalCompileShaderMSC::DoCompileMetalShader(
 	const FShaderCompilerInput& Input,
 	FShaderCompilerOutput& Output,
 	const FString& InPreprocessedShader,
-	FSHAHash GUIDHash,
 	uint32 VersionEnum,
 	EMetalGPUSemantics Semantics,
 	uint32 MaxUnrollLoops,
@@ -806,7 +804,7 @@ void FMetalCompileShaderMSC::DoCompileMetalShader(
 	if (Result != 0)
 	{
 		Output.Target = Input.Target;
-		BuildMetalShaderOutput(Output, Input, GUIDHash, MetalSource.c_str(), MetalSource.length(), CRCLen, CRC, VersionEnum, *Standard, *MinOSVersion, Output.Errors, OutputData.TypedBuffers, OutputData.InvariantBuffers, OutputData.TypedUAVs, OutputData.ConstantBuffers, bAllowFastIntrinsics
+		BuildMetalShaderOutput(Output, Input, MetalSource.c_str(), MetalSource.length(), CRCLen, CRC, VersionEnum, *Standard, *MinOSVersion, Output.Errors, OutputData.TypedBuffers, OutputData.InvariantBuffers, OutputData.TypedUAVs, OutputData.ConstantBuffers, bAllowFastIntrinsics
 		  , NumCBVs, OutputSizeVS, MaxInputPrimitivesPerMeshThreadgroupGS, bUsesDiscard, ReflectionJSON, MetalBytecode
 		);
 	}
@@ -831,7 +829,6 @@ void FMetalCompileShaderMSC::DoCompileMetalShader(
 	const FShaderCompilerInput& Input,
 	FShaderCompilerOutput& Output,
 	const FString& InPreprocessedShader,
-	FSHAHash GUIDHash,
 	uint32 VersionEnum,
 	EMetalGPUSemantics Semantics,
 	uint32 MaxUnrollLoops,
diff --git a/Engine/Source/Developer/Apple/MetalShaderFormat/Private/MetalCompileShaderMSC.h b/Engine/Source/Developer/Apple/MetalShaderFormat/Private/MetalCompileShaderMSC.h
--- a/Engine/Source/Developer/Apple/MetalShaderFormat/Private/MetalCompileShaderMSC.h
+++ b/Engine/Source/Developer/Apple/MetalShaderFormat/Private/MetalCompileShaderMSC.h
@@ -17,7 +17,6 @@ public:
 									 const FShaderCompilerInput& Input,
 									 FShaderCompilerOutput& Output,
 									 const FString& InPreprocessedShader,
-									 FSHAHash GUIDHash,
 									 uint32 VersionEnum,
 									 EMetalGPUSemantics Semantics,
 									 uint32 MaxUnrollLoops,
diff --git a/Engine/Source/Developer/Apple/MetalShaderFormat/Private/MetalCompileShaderSPIRV.cpp b/Engine/Source/Developer/Apple/MetalShaderFormat/Private/MetalCompileShaderSPIRV.cpp
--- a/Engine/Source/Developer/Apple/MetalShaderFormat/Private/MetalCompileShaderSPIRV.cpp
+++ b/Engine/Source/Developer/Apple/MetalShaderFormat/Private/MetalCompileShaderSPIRV.cpp
@@ -17,7 +17,6 @@
 extern void BuildMetalShaderOutput(
 	FShaderCompilerOutput& ShaderOutput,
 	const FShaderCompilerInput& ShaderInput,
-	FSHAHash const& GUIDHash,
 	const ANSICHAR* InShaderSource,
 	uint32 SourceLen,
 	uint32 SourceCRCLen,
@@ -116,7 +115,6 @@ void FMetalCompileShaderSPIRV::DoCompileMetalShader(
 	const FShaderCompilerInput& Input,
 	FShaderCompilerOutput& Output,
 	const FString& InPreprocessedShader,
-	FSHAHash GUIDHash,
 	uint32 VersionEnum,
 	EMetalGPUSemantics Semantics,
 	uint32 MaxUnrollLoops,
@@ -1298,7 +1296,7 @@ void FMetalCompileShaderSPIRV::DoCompileMetalShader(
 	if (Result != 0)
 	{
 		Output.Target = Input.Target;
-		BuildMetalShaderOutput(Output, Input, GUIDHash, MetalSource.c_str(), MetalSource.length(), CRCLen, CRC, VersionEnum, *Standard, *MinOSVersion, Output.Errors, OutputData.TypedBuffers, OutputData.InvariantBuffers, OutputData.TypedUAVs, OutputData.ConstantBuffers, bAllowFastIntrinsics
+		BuildMetalShaderOutput(Output, Input, MetalSource.c_str(), MetalSource.length(), CRCLen, CRC, VersionEnum, *Standard, *MinOSVersion, Output.Errors, OutputData.TypedBuffers, OutputData.InvariantBuffers, OutputData.TypedUAVs, OutputData.ConstantBuffers, bAllowFastIntrinsics
 #if UE_METAL_USE_METAL_SHADER_CONVERTER
 		  , 0, 0, 0, false, nullptr, FMetalShaderBytecode()
 #endif
diff --git a/Engine/Source/Developer/Apple/MetalShaderFormat/Private/MetalCompileShaderSPIRV.h b/Engine/Source/Developer/Apple/MetalShaderFormat/Private/MetalCompileShaderSPIRV.h
--- a/Engine/Source/Developer/Apple/MetalShaderFormat/Private/MetalCompileShaderSPIRV.h
+++ b/Engine/Source/Developer/Apple/MetalShaderFormat/Private/MetalCompileShaderSPIRV.h
@@ -14,7 +14,6 @@ public:
 									 const FShaderCompilerInput& Input,
 									 FShaderCompilerOutput& Output,
 									 const FString& InPreprocessedShader,
-									 FSHAHash GUIDHash,
 									 uint32 VersionEnum,
 									 EMetalGPUSemantics Semantics,
 									 uint32 MaxUnrollLoops,
diff --git a/Engine/Source/Developer/Apple/MetalShaderFormat/Private/MetalShaderCompiler.cpp b/Engine/Source/Developer/Apple/MetalShaderFormat/Private/MetalShaderCompiler.cpp
--- a/Engine/Source/Developer/Apple/MetalShaderFormat/Private/MetalShaderCompiler.cpp
+++ b/Engine/Source/Developer/Apple/MetalShaderFormat/Private/MetalShaderCompiler.cpp
@@ -232,7 +232,6 @@ static const int32 Str##PrefixLen = FCStringAnsi::Strlen(Str##Prefix)
 void BuildMetalShaderOutput(
 	FShaderCompilerOutput& ShaderOutput,
 	const FShaderCompilerInput& ShaderInput,
-	FSHAHash const& GUIDHash,
 	const ANSICHAR* InShaderSource,
 	uint32 SourceLen,
 	uint32 SourceCRCLen,
@@ -737,7 +736,6 @@ void BuildMetalShaderOutput(
 			FMetalShaderBytecodeJob Job;
 			Job.IncludeDir = TempDir;
 			Job.ShaderFormat = ShaderInput.ShaderFormat;
-			Job.Hash = GUIDHash;
 			Job.TmpFolder = TempDir;
 			Job.InputFile = MetalFileName;
 			Job.OutputFile = MetallibFileName;
@@ -1066,31 +1064,17 @@ void CompileMetalShader(const FShaderCompilerInput& Input, const FShaderPreproce
 	}
 
 
-	FSHAHash GUIDHash;
-	if (!EnumHasAnyFlags(Input.DebugInfoFlags, EShaderDebugInfoFlags::CompileFromDebugUSF))
-	{
-		TArray<FString> GUIDFiles;
-		GUIDFiles.Add(FPaths::ConvertRelativePathToFull(TEXT("/Engine/Public/Platform/Metal/MetalCommon.ush")));
-		GUIDFiles.Add(FPaths::ConvertRelativePathToFull(TEXT("/Engine/Public/ShaderVersion.ush")));
-		GUIDHash = GetShaderFilesHash(GUIDFiles, Input.Target.GetPlatform());
-	}
-	else
-	{
-		FGuid Guid = FGuid::NewGuid();
-		FSHA1::HashBuffer(&Guid, sizeof(FGuid), GUIDHash.Hash);
-	}
-
 #if UE_METAL_USE_METAL_SHADER_CONVERTER
 	const bool bBindlessEnabled = (Input.Environment.CompilerFlags.Contains(CFLAG_BindlessResources) || Input.Environment.CompilerFlags.Contains(CFLAG_BindlessSamplers));
 	
 	if(bBindlessEnabled && Input.ShaderFormat == NAME_SF_METAL_SM6)
 	{
-		FMetalCompileShaderMSC::DoCompileMetalShader(Input, Output, PreprocessedSource, GUIDHash, VersionEnum, Semantics, MaxUnrollLoops, (EShaderFrequency)Input.Target.Frequency, bDumpDebugInfo, Standard, MinOSVersion);
+		FMetalCompileShaderMSC::DoCompileMetalShader(Input, Output, PreprocessedSource, VersionEnum, Semantics, MaxUnrollLoops, (EShaderFrequency)Input.Target.Frequency, bDumpDebugInfo, Standard, MinOSVersion);
 	}
 	else
 #endif
 	{
-		FMetalCompileShaderSPIRV::DoCompileMetalShader(Input, Output, PreprocessedSource, GUIDHash, VersionEnum, Semantics, MaxUnrollLoops, (EShaderFrequency)Input.Target.Frequency, bDumpDebugInfo, Standard, MinOSVersion);
+		FMetalCompileShaderSPIRV::DoCompileMetalShader(Input, Output, PreprocessedSource, VersionEnum, Semantics, MaxUnrollLoops, (EShaderFrequency)Input.Target.Frequency, bDumpDebugInfo, Standard, MinOSVersion);
 	}
 	ShaderParameterParser.ValidateShaderParameterTypes(Input, bIsMobile, Output);
 }
diff --git a/Engine/Source/Developer/Apple/MetalShaderFormat/Private/MetalShaderCompiler.h b/Engine/Source/Developer/Apple/MetalShaderFormat/Private/MetalShaderCompiler.h
--- a/Engine/Source/Developer/Apple/MetalShaderFormat/Private/MetalShaderCompiler.h
+++ b/Engine/Source/Developer/Apple/MetalShaderFormat/Private/MetalShaderCompiler.h
@@ -41,7 +41,6 @@ struct FMetalShaderDebugInfo
 struct FMetalShaderBytecodeJob
 {
 	FName ShaderFormat;
-	FSHAHash Hash;
     FString Defines;
 	FString TmpFolder;
 	FString InputFile;
