[shaders] add error code to spirv reflect failure assert

diff --git a/Engine/Source/Developer/Apple/MetalShaderFormat/Private/MetalCompileShaderSPIRV.cpp b/Engine/Source/Developer/Apple/MetalShaderFormat/Private/MetalCompileShaderSPIRV.cpp
--- a/Engine/Source/Developer/Apple/MetalShaderFormat/Private/MetalCompileShaderSPIRV.cpp
+++ b/Engine/Source/Developer/Apple/MetalShaderFormat/Private/MetalCompileShaderSPIRV.cpp
@@ -246,7 +246,7 @@ void FMetalCompileShaderSPIRV::DoCompileMetalShader(
 			// Now perform reflection on the SPIRV and tweak any decorations that we need to.
 			// This used to be done via JSON, but that was slow and alloc happy so use SPIRV-Reflect instead.
 			spv_reflect::ShaderModule Reflection(SpirvData.Num() * sizeof(uint32), SpirvData.GetData());
-			check(Reflection.GetResult() == SPV_REFLECT_RESULT_SUCCESS);
+			checkf(Reflection.GetResult() == SPV_REFLECT_RESULT_SUCCESS, TEXT("SPIRV-reflect failed with error %d"), Reflection.GetResult());
 			
 			SpvReflectResult SPVRResult = SPV_REFLECT_RESULT_NOT_READY;
 			uint32 Count = 0;
