Change the CHECK_METAL_COMPILER_TOOLCHAIN_SETUP define to a cvar Metal.CheckCompilerToolChainSetup

diff --git a/Engine/Source/Developer/Apple/MetalShaderFormat/Private/MetalShaderFormat.cpp b/Engine/Source/Developer/Apple/MetalShaderFormat/Private/MetalShaderFormat.cpp
--- a/Engine/Source/Developer/Apple/MetalShaderFormat/Private/MetalShaderFormat.cpp
+++ b/Engine/Source/Developer/Apple/MetalShaderFormat/Private/MetalShaderFormat.cpp
@@ -26,8 +26,15 @@ DEFINE_LOG_CATEGORY(LogMetalShaderCompiler)
 
 #define WRITE_METAL_SHADER_SOURCE_ARCHIVE 0
 
-// Set this define to get additional logging information about Metal toolchain setup.
-#define CHECK_METAL_COMPILER_TOOLCHAIN_SETUP 0
+// Set this cvar to get additional logging information about Metal toolchain setup.
+static int32 GCheckCompilerToolChainSetup = 0;
+static FAutoConsoleVariableRef CVarCheckCompilerToolChainSetup(
+	TEXT("Metal.CheckCompilerToolChainSetup"),
+	GCheckCompilerToolChainSetup,
+	TEXT("Should we check the Metal Compiler ToolChain Setup."),
+	ECVF_Default
+);
+
 
 extern bool PreprocessMetalShader(const FShaderCompilerInput& Input, const FShaderCompilerEnvironment& Environment, FShaderPreprocessOutput& PreprocessOutput);
 extern void CompileMetalShader(const FShaderCompilerInput& Input, const FShaderPreprocessOutput& InPreprocessOutput, FShaderCompilerOutput& Output);
@@ -687,9 +694,10 @@ void FMetalCompilerToolchain::Init()
 
 	if (Result != EMetalToolchainStatus::Success)
 	{
-#if CHECK_METAL_COMPILER_TOOLCHAIN_SETUP
-		UE_LOG(LogMetalCompilerSetup, Warning, TEXT("Metal compiler not found. Shaders will be stored as text."));
-#endif
+		if (GCheckCompilerToolChainSetup > 0)
+		{
+			UE_LOG(LogMetalCompilerSetup, Warning, TEXT("Metal compiler not found. Shaders will be stored as text."));
+		}
 		bToolchainAvailable = false;
 	}
 	else
@@ -719,28 +727,29 @@ void FMetalCompilerToolchain::Init()
 		bToolchainAvailable = true;
 	}
 
-#if CHECK_METAL_COMPILER_TOOLCHAIN_SETUP
-	if (Result == EMetalToolchainStatus::Success)
+	if (GCheckCompilerToolChainSetup > 0)
 	{
-		check(IsCompilerAvailable());
-		UE_LOG(LogMetalCompilerSetup, Log, TEXT("Metal toolchain setup complete."));
-		UE_LOG(LogMetalCompilerSetup, Log, TEXT("Using Local Metal compiler"));
-		if (!MetalFrontendBinaryCommand[AppleSDKMac].IsEmpty())
+		if (Result == EMetalToolchainStatus::Success)
 		{
-			UE_LOG(LogMetalCompilerSetup, Log, TEXT("Mac metalfe found at %s"), *MetalFrontendBinaryCommand[AppleSDKMac]);
+			check(IsCompilerAvailable());
+			UE_LOG(LogMetalCompilerSetup, Log, TEXT("Metal toolchain setup complete."));
+			UE_LOG(LogMetalCompilerSetup, Log, TEXT("Using Local Metal compiler"));
+			if (!MetalFrontendBinaryCommand[AppleSDKMac].IsEmpty())
+			{
+				UE_LOG(LogMetalCompilerSetup, Log, TEXT("Mac metalfe found at %s"), *MetalFrontendBinaryCommand[AppleSDKMac]);
+			}
+			if (!MetalFrontendBinaryCommand[AppleSDKMobile].IsEmpty())
+			{
+				UE_LOG(LogMetalCompilerSetup, Log, TEXT("Mobile metalfe found at %s"), *MetalFrontendBinaryCommand[AppleSDKMobile]);
+			}
+			UE_LOG(LogMetalCompilerSetup, Log, TEXT("Mac metalfe version %s"), *MetalCompilerVersionString[AppleSDKMac]);
+			UE_LOG(LogMetalCompilerSetup, Log, TEXT("Mobile metalfe version %s"), *MetalCompilerVersionString[AppleSDKMobile]);
 		}
-		if (!MetalFrontendBinaryCommand[AppleSDKMobile].IsEmpty())
+		else
 		{
-			UE_LOG(LogMetalCompilerSetup, Log, TEXT("Mobile metalfe found at %s"), *MetalFrontendBinaryCommand[AppleSDKMobile]);
+			UE_LOG(LogMetalCompilerSetup, Warning, TEXT("Failed to set up Metal toolchain. See log above. Shaders will not be compiled offline."));
 		}
-		UE_LOG(LogMetalCompilerSetup, Log, TEXT("Mac metalfe version %s"), *MetalCompilerVersionString[AppleSDKMac]);
-		UE_LOG(LogMetalCompilerSetup, Log, TEXT("Mobile metalfe version %s"), *MetalCompilerVersionString[AppleSDKMobile]);
 	}
-	else
-	{
-		 UE_LOG(LogMetalCompilerSetup, Warning, TEXT("Failed to set up Metal toolchain. See log above. Shaders will not be compiled offline."));
-	}
-#endif
 }
 
 void FMetalCompilerToolchain::Teardown()
@@ -941,10 +950,11 @@ FMetalCompilerToolchain::EMetalToolchainStatus FMetalCompilerToolchain::DoWindow
 	bool bUseLocalMetalToolchain = FPaths::FileExists(MetalFrontendBinaryCommand[AppleSDKMac]) && FPaths::FileExists(MetalFrontendBinaryCommand[AppleSDKMobile]);
 	if (!bUseLocalMetalToolchain)
 	{
-#if CHECK_METAL_COMPILER_TOOLCHAIN_SETUP
-		UE_LOG(LogMetalCompilerSetup, Display, TEXT("Searching for Metal toolchain, but it doesn't appear to be installed."));
-		UE_LOG(LogMetalCompilerSetup, Display, TEXT("Searched for %s and %s"), *MetalFrontendBinaryCommand[AppleSDKMac], *MetalFrontendBinaryCommand[AppleSDKMobile]);
-#endif
+		if (GCheckCompilerToolChainSetup > 0)
+		{
+			UE_LOG(LogMetalCompilerSetup, Display, TEXT("Searching for Metal toolchain, but it doesn't appear to be installed."));
+			UE_LOG(LogMetalCompilerSetup, Display, TEXT("Searched for %s and %s"), *MetalFrontendBinaryCommand[AppleSDKMac], *MetalFrontendBinaryCommand[AppleSDKMobile]);
+		}
 		return EMetalToolchainStatus::ToolchainNotFound;
 	}
 
@@ -964,9 +974,10 @@ FMetalCompilerToolchain::EMetalToolchainStatus FMetalCompilerToolchain::DoWindow
         !FPaths::FileExists(AirPackBinaryCommand[AppleSDKMac]) ||
         !FPaths::FileExists(AirPackBinaryCommand[AppleSDKMobile]))
 	{
-#if CHECK_METAL_COMPILER_TOOLCHAIN_SETUP
-		UE_LOG(LogMetalCompilerSetup, Warning, TEXT("Missing toolchain binaries."))
-#endif
+		if (GCheckCompilerToolChainSetup > 0)
+		{
+			UE_LOG(LogMetalCompilerSetup, Warning, TEXT("Missing toolchain binaries."))
+		}
 		return EMetalToolchainStatus::ToolchainNotFound;
 	}
 
