[Shaders] Only store optional shader source when CFLAG_ExtraShaderData is specified, not for CFLAG_Archive.

Storing the whole metal shader source text in the optional data field was introduced with CL 3341849 for debugging purposes.
This should only be required when CFLAG_ExtraShaderData is specified (enabled via CVar r.Shaders.ExtraData).
The only place this optional data is needed outside the shader compiler, is TMetalBaseShader::Init() to initialize the field "GlslCodeNSString" which is documented as debuggable text.

Also invalidate all Metal shaders to trigger a recompilation with a much smaller memory footprint of the output shader binaries.

#rnx

diff --git a/Engine/Shaders/Public/Platform/Metal/MetalCommon.ush b/Engine/Shaders/Public/Platform/Metal/MetalCommon.ush
--- a/Engine/Shaders/Public/Platform/Metal/MetalCommon.ush
+++ b/Engine/Shaders/Public/Platform/Metal/MetalCommon.ush
@@ -7,7 +7,7 @@
 #pragma once
 
 // Update this GUID to invalidate shader recompilation for only Metal shaders
-#pragma message("UESHADERMETADATA_VERSION E4FA73AD-5066-4067-8B47-814FDC0D6332")
+#pragma message("UESHADERMETADATA_VERSION D04FC893-6A55-4F32-A961-9CD15F4585CD")
 
 #if VERTEXSHADER && HULLSHADER
 #ifndef TESSELLATIONSHADER
diff --git a/Engine/Source/Developer/Apple/MetalShaderFormat/Private/MetalShaderCompiler.cpp b/Engine/Source/Developer/Apple/MetalShaderFormat/Private/MetalShaderCompiler.cpp
--- a/Engine/Source/Developer/Apple/MetalShaderFormat/Private/MetalShaderCompiler.cpp
+++ b/Engine/Source/Developer/Apple/MetalShaderFormat/Private/MetalShaderCompiler.cpp
@@ -816,11 +816,6 @@ void BuildMetalShaderOutput(
 					ShaderOutput.ShaderCode.AddOptionalData(EShaderOptionalDataKey::NativePath, TCHAR_TO_UTF8(*Bytecode.NativePath));
 				}
 			}
-			else if (ShaderInput.Environment.CompilerFlags.Contains(CFLAG_Archive))
-			{
-				ShaderOutput.ShaderCode.AddOptionalData(EShaderOptionalDataKey::SourceCode, TCHAR_TO_UTF8(*MetalCode));
-				ShaderOutput.ShaderCode.AddOptionalData(EShaderOptionalDataKey::NativePath, TCHAR_TO_UTF8(*Bytecode.NativePath));
-			}
 			
 			ShaderOutput.NumTextureSamplers = Header.Bindings.NumSamplers;
 		}
