Fix half precision support for iOS

diff --git a/Engine/Shaders/Public/Platform.ush b/Engine/Shaders/Public/Platform.ush
--- a/Engine/Shaders/Public/Platform.ush
+++ b/Engine/Shaders/Public/Platform.ush
@@ -322,6 +322,10 @@
 #define FORCE_FLOATS 0
 #endif
 
+#if COMPILER_DXC && !FORCE_FLOATS && (ES3_1_PROFILE || METAL_PROFILE)
+    #pragma dxc diagnostic ignored "-Wconversion"
+#endif
+
 #if ((!(ES3_1_PROFILE || METAL_PROFILE)) && !PLATFORM_SUPPORTS_REAL_TYPES) || FORCE_FLOATS
 	// Always use floats when using the ES3/METAL compiler, because platforms not optimized for lower precision, 
 	// And we don't want potential side effects on other platforms
@@ -475,7 +479,11 @@
 	#define UB_INT(Dim)         int##Dim
 	#define UB_UINT(Dim)        uint##Dim
 	#define UB_FLOAT(Dim)       float##Dim
-	#define UB_HALF_FLOAT(Dim)  half##Dim
+    #if METAL_PROFILE
+        #define UB_HALF_FLOAT(Dim)  float##Dim
+    #else
+        #define UB_HALF_FLOAT(Dim)  half##Dim
+    #endif
 	#define UB_FIXED_FLOAT(Dim) fixed##Dim
 #else
 	#define UB_INT(Dim)         int##Dim
diff --git a/Engine/Shaders/Public/Platform/Metal/MetalCommon.ush b/Engine/Shaders/Public/Platform/Metal/MetalCommon.ush
--- a/Engine/Shaders/Public/Platform/Metal/MetalCommon.ush
+++ b/Engine/Shaders/Public/Platform/Metal/MetalCommon.ush
@@ -7,7 +7,7 @@
 #pragma once
 
 // Update this GUID to invalidate shader recompilation for only Metal shaders
-#pragma message("UESHADERMETADATA_VERSION E4FA73AD-5066-4067-8B47-814FDC0D6332")
+#pragma message("UESHADERMETADATA_VERSION 8C6A7A49-B8C0-4561-9D8B-1E0C755D93F7")
 
 #if VERTEXSHADER && HULLSHADER
 #ifndef TESSELLATIONSHADER
diff --git a/Engine/Source/Developer/Apple/MetalShaderFormat/Private/MetalCompileShaderMSC.cpp b/Engine/Source/Developer/Apple/MetalShaderFormat/Private/MetalCompileShaderMSC.cpp
--- a/Engine/Source/Developer/Apple/MetalShaderFormat/Private/MetalCompileShaderMSC.cpp
+++ b/Engine/Source/Developer/Apple/MetalShaderFormat/Private/MetalCompileShaderMSC.cpp
@@ -15,6 +15,7 @@
 #include "ShaderCompilerDefinitions.h"
 #include "SpirvReflectCommon.h"
 #include "ShaderParameterParser.h"
+#include "Containers/AnsiString.h"
 
 #include <regex>
 
@@ -516,7 +517,7 @@ void FMetalCompileShaderMSC::DoCompileMetalShader(
 		CompilerContext.LoadSource(PreprocessedShader, Input.VirtualSourceFilePath, Input.EntryPointName, Frequency);
 
 		// Convert shader source to ANSI string
-		std::string SourceData(CompilerContext.GetSourceString(), static_cast<size_t>(CompilerContext.GetSourceLength()));
+		FAnsiString SourceData = FAnsiString::ConstructFromPtrSize(CompilerContext.GetSourceString(), CompilerContext.GetSourceLength());
 
 		// Replace special case texture "gl_LastFragData" by native subpass fetch operation
 		static const uint32 MaxMetalSubpasses = 8;
@@ -527,12 +528,12 @@ void FMetalCompileShaderMSC::DoCompileMetalShader(
 		// If source data was modified, reload it into the compiler context
 		if (bSourceDataWasModified)
 		{
-			CompilerContext.LoadSource(FAnsiStringView(SourceData.c_str(), SourceData.length()), Input.VirtualSourceFilePath, Input.EntryPointName, Frequency);
+			CompilerContext.LoadSource(SourceData, Input.VirtualSourceFilePath, Input.EntryPointName, Frequency);
 		}
 
 		if (bDumpDebugInfo)
 		{
-			DumpDebugShaderText(Input, &SourceData[0], SourceData.size(), TEXT("rewritten.hlsl"));
+			DumpDebugShaderText(Input, &SourceData[0], SourceData.Len(), TEXT("rewritten.hlsl"));
 		}
 		
 		CrossCompiler::FHlslccHeaderWriter CCHeaderWriter;
@@ -630,7 +631,7 @@ void FMetalCompileShaderMSC::DoCompileMetalShader(
 #endif
 		 
 		// TODO: Is there a flag we could check to avoid this string lookup?
-		bool bUsesDualSourceBlending = (SourceData.find("vk::location") != std::string::npos);
+		bool bUsesDualSourceBlending = (SourceData.Find("vk::location") != INDEX_NONE);
 		if (bUsesDualSourceBlending)
 		{
 			IRCompilerSetDualSourceBlendingConfiguration(CompilerInstance, IRDualSourceBlendingConfigurationForceEnabled);
diff --git a/Engine/Source/Developer/Apple/MetalShaderFormat/Private/MetalCompileShaderSPIRV.cpp b/Engine/Source/Developer/Apple/MetalShaderFormat/Private/MetalCompileShaderSPIRV.cpp
--- a/Engine/Source/Developer/Apple/MetalShaderFormat/Private/MetalCompileShaderSPIRV.cpp
+++ b/Engine/Source/Developer/Apple/MetalShaderFormat/Private/MetalCompileShaderSPIRV.cpp
@@ -11,6 +11,7 @@
 #include "SpirvReflectCommon.h"
 #include "ShaderParameterParser.h"
 #include "Misc/OutputDeviceRedirector.h"
+#include "Containers/AnsiString.h"
 
 #include <regex>
 
@@ -40,75 +41,135 @@ extern void BuildMetalShaderOutput(
 #endif
 );
 
-static void Patch16bitInHlslSource(const FShaderCompilerInput& Input, std::string& SourceData)
-{
-	static const std::string TextureTypes [] = {
-		"Texture1D",
-		"Texture1DArray",
-		"Texture2D",
-		"Texture2DArray",
-		"Texture3D",
-		"TextureCube",
-		"TextureCubeArray",
-		"Buffer"
-	};
-	
-	// half precision textures and buffers are not supported in DXC
-	for(uint32_t i = 0; i < UE_ARRAY_COUNT(TextureTypes); ++i)
-	{
-		const std::string & TextureTypeString = TextureTypes[i];
-		
-		std::regex pattern(TextureTypeString + "<\\s?half");
-		SourceData = std::regex_replace(SourceData, pattern, TextureTypeString + "<float");
-	}
+static void Patch16bitInHlslSource(const FShaderCompilerInput& Input, FAnsiString& SourceData)
+{	
+	TArray<int32> PatchPositions;
 	
-	static const std::string ConstHalf = "const half";
-	static const std::string ConstFloat = "const float";
+	static const FAnsiStringView TextureStr = "Texture";
+	static const FAnsiStringView HalfStr("half");
+	static const FAnsiStringView FloatStr("float");
 	
-	// Replace half in constant buffers to use float
-	for (const TPair<FString, FUniformBufferEntry> & Pair : Input.Environment.UniformBufferMap)
+	// Find all half texture and buffer types to patch
 	{
-		std::string CBufferName = std::string("cbuffer ") + TCHAR_TO_UTF8(*Pair.Key);
-		
-		size_t StructPos = SourceData.find(CBufferName);
-		if(StructPos != std::string::npos)
+		int32 Pos = SourceData.Find(TextureStr );
+		while (Pos != INDEX_NONE)
 		{
-			size_t StructEndPos = SourceData.find("};", StructPos);
-			if(StructEndPos != std::string::npos)
+			Pos += TextureStr.Len();
+			if (Pos >= SourceData.Len())
 			{
-				TArray<size_t> HalfPositions;
-				size_t HalfPos = SourceData.find(ConstHalf, StructPos);
-				
-				while(HalfPos != std::string::npos &&
-					  HalfPos < StructEndPos)
-				{
-					HalfPositions.Add(HalfPos);
-					HalfPos = SourceData.find(ConstHalf, HalfPos + ConstHalf.size());
-				}
-				
-				for(int32_t i = HalfPositions.Num()-1; i >= 0; i--)
+				break;
+			}
+			static const FAnsiStringView TextureTypeNames[] =
+			{
+				"1D<",
+				"1DArray<",
+				"2D<",
+				"2DArray<",
+				"3D<",
+				"Cube<",
+				"CubeArray<"
+			};
+			
+			FAnsiStringView CodeStr(&SourceData[Pos], SourceData.Len() - Pos);
+			for (const FAnsiStringView& TextureTypeStr : TextureTypeNames)
+			{
+				if (CodeStr.StartsWith(TextureTypeStr))
 				{
-					SourceData.replace(HalfPositions[i], ConstHalf.size(), ConstFloat);
+					CodeStr.RightChopInline(TextureTypeStr.Len());
+					CodeStr.TrimStartInline();
+					if (CodeStr.StartsWith(HalfStr))
+					{
+						PatchPositions.Push(CodeStr.GetData() - *SourceData);
+					}
 				}
 			}
+			
+			Pos = SourceData.Find(TextureStr, ESearchCase::CaseSensitive, ESearchDir::FromStart, Pos);
 		}
 	}
 	
-	// Replace Globals
-	size_t GlobalPos = SourceData.find(std::string("\n") + ConstHalf);
-	while(GlobalPos != std::string::npos)
 	{
-		// Check this is a global and not an assignment
-		size_t LineEndPos = SourceData.find(";", GlobalPos);
-		size_t AssignmentPos = SourceData.find("=", GlobalPos);
-		
-		if(AssignmentPos == std::string::npos || AssignmentPos > LineEndPos)
+		static const FAnsiStringView BufferStr = "Buffer<";
+		int32 Pos = SourceData.Find(BufferStr);
+		while (Pos != INDEX_NONE)
 		{
-			SourceData.replace(GlobalPos+1, ConstHalf.size(), ConstFloat);
+			Pos += BufferStr.Len();
+			if (Pos >= SourceData.Len())
+			{
+				break;
+			}
+
+			FAnsiStringView CodeStr(&SourceData[Pos], SourceData.Len() - Pos);
+			CodeStr.TrimStartInline();
+			if (CodeStr.StartsWith(HalfStr))
+			{
+				PatchPositions.Push(CodeStr.GetData() - *SourceData);
+			}
+			Pos = SourceData.Find(BufferStr, ESearchCase::CaseSensitive, ESearchDir::FromStart, Pos);
 		}
-		
-		GlobalPos = SourceData.find(std::string("\n") + ConstHalf, GlobalPos+ConstHalf.size());
 	}
+	
+	// Convert global uniforms to float (uniforms in uniform buffers are always float via shader macro UB_HALF_FLOAT)
+	FAnsiStringView CodeStr = SourceData;
+	while(CodeStr.Len() > 0)
+	{
+		CodeStr.TrimStartInline();
+		if (CodeStr.StartsWith(HalfStr))
+		{
+			int32 PatchPos = CodeStr.GetData() - *SourceData;
+			CodeStr.RemovePrefix(HalfStr.Len());
+			char Head = CodeStr[0];
+			if (Head == '1' || Head == '2' || Head == '3' || Head == '4')
+			{
+				CodeStr.RemovePrefix(1);
+				Head = CodeStr[0];
+			}
+			if (isspace(Head))
+			{
+				CodeStr.TrimStartInline();
+				Head = CodeStr[0];
+				if (isalpha(Head) || Head == '_')
+				{
+					for (char C : CodeStr)
+					{
+						if (isalnum(C) || isspace(C) || C == '_')
+						{
+							continue;
+						}
+						else if (C == ';')
+						{
+							PatchPositions.Push(PatchPos);
+							break;
+						}
+						else
+						{
+							break;
+						}
+					}
+				}
+			}
+		}
+		// This is brittle, but only consider "half" types at the beginning of a new line of code.
+		// We only want to replace global uniforms, not local variables or variables inside uniform buffers.
+		int32 FindPos = CodeStr.Find("\nhalf", 1);
+		if (FindPos == INDEX_NONE)
+		{
+			break;
+		}
+		CodeStr.RemovePrefix(FindPos + 1);
+	}
+
+	// Create new code string where all relevant "half" strings are replaced with "float"
+	FAnsiString Result = FAnsiString::ConstructWithSlack("", SourceData.Len() + PatchPositions.Num());
+	int32 LastEndPatchPos = 0;
+	for (const int32 Pos : PatchPositions)
+	{
+		Result.Append(&SourceData[LastEndPatchPos], Pos - LastEndPatchPos);
+		Result.Append(FloatStr);
+		LastEndPatchPos = Pos + HalfStr.Len();
+	}
+	Result.Append(&SourceData[LastEndPatchPos], SourceData.Len() - LastEndPatchPos);
+	SourceData = std::move(Result);
 }
 
 void FMetalCompileShaderSPIRV::DoCompileMetalShader(
@@ -197,12 +258,17 @@ void FMetalCompileShaderSPIRV::DoCompileMetalShader(
 		CompilerContext.LoadSource(PreprocessedShader, Input.VirtualSourceFilePath, Input.EntryPointName, Frequency);
 
 		// Convert shader source to ANSI string
-		std::string SourceData(CompilerContext.GetSourceString(), static_cast<size_t>(CompilerContext.GetSourceLength()));
+		FAnsiString SourceData = FAnsiString::ConstructFromPtrSize(CompilerContext.GetSourceString(), CompilerContext.GetSourceLength());
 
 		// Replace special case texture "gl_LastFragData" by native subpass fetch operation
 		static const uint32 MaxMetalSubpasses = 8;
 		uint32 SubpassInputsDim[MaxMetalSubpasses];
 
+		if (bDumpDebugInfo)
+		{
+			DumpDebugShaderText(Input, &SourceData[0], SourceData.Len(), TEXT("orig.hlsl"));
+		}
+		
 		bool bSourceDataWasModified = PatchSpecialTextureInHlslSource(SourceData, SubpassInputsDim, MaxMetalSubpasses);
 		
 		// If using 16 bit types disable half precision in constant buffer due to errors in layout
@@ -215,12 +281,12 @@ void FMetalCompileShaderSPIRV::DoCompileMetalShader(
 		// If source data was modified, reload it into the compiler context
 		if (bSourceDataWasModified)
 		{
-			CompilerContext.LoadSource(FAnsiStringView(SourceData.c_str(), SourceData.length()), Input.VirtualSourceFilePath, Input.EntryPointName, Frequency);
+			CompilerContext.LoadSource(FAnsiStringView(*SourceData, SourceData.Len()), Input.VirtualSourceFilePath, Input.EntryPointName, Frequency);
 		}
 
 		if (bDumpDebugInfo)
 		{
-			DumpDebugShaderText(Input, &SourceData[0], SourceData.size(), TEXT("rewritten.hlsl"));
+			DumpDebugShaderText(Input, &SourceData[0], SourceData.Len(), TEXT("rewritten.hlsl"));
 		}
 		
 		CrossCompiler::FHlslccHeaderWriter CCHeaderWriter;
diff --git a/Engine/Source/Developer/Apple/MetalShaderFormat/Private/MetalShaderCompiler.cpp b/Engine/Source/Developer/Apple/MetalShaderFormat/Private/MetalShaderCompiler.cpp
--- a/Engine/Source/Developer/Apple/MetalShaderFormat/Private/MetalShaderCompiler.cpp
+++ b/Engine/Source/Developer/Apple/MetalShaderFormat/Private/MetalShaderCompiler.cpp
@@ -1547,8 +1547,14 @@ bool FinalizeLibrary_Metal(FName const& Format, FString const& WorkingDir, FStri
 	return bOK;
 }
 
+static void ReplaceString(FAnsiString& Str, int32 Pos, int Len, const FAnsiString& NewStr)
+{
+	Str.RemoveAt(Pos, Len, EAllowShrinking::No);
+	Str.InsertAt(Pos, NewStr);
+}
+
 // Replace the special texture "gl_LastFragData" to a native subpass fetch operation. Returns true if the input source has been modified.
-bool PatchSpecialTextureInHlslSource(std::string& SourceData, uint32* OutSubpassInputsDim, uint32 SubpassInputDimCount)
+bool PatchSpecialTextureInHlslSource(FAnsiString& SourceData, uint32* OutSubpassInputsDim, uint32 SubpassInputDimCount)
 {
 	bool bSourceDataWasModified = false;
 
@@ -1556,13 +1562,13 @@ bool PatchSpecialTextureInHlslSource(std::string& SourceData, uint32* OutSubpass
 	FMemory::Memzero(OutSubpassInputsDim, sizeof(uint32) * SubpassInputDimCount);
 	
 	// Check if special texture is present in the code
-	static const std::string GSpecialTextureLastFragData = "gl_LastFragData";
-	if (SourceData.find(GSpecialTextureLastFragData) != std::string::npos)
+	const char* GSpecialTextureLastFragData = "gl_LastFragData";
+	if (SourceData.Find(GSpecialTextureLastFragData, ESearchCase::CaseSensitive) != INDEX_NONE)
 	{
 		struct FHlslVectorType
 		{
-			std::string TypenameIdent;
-			std::string TypenameSuffix;
+			const char* TypenameIdent;
+			const char* TypenameSuffix;
 			uint32 Dimension;
 		};
 		const FHlslVectorType FragDeclTypes[4] =
@@ -1579,36 +1585,34 @@ bool PatchSpecialTextureInHlslSource(std::string& SourceData, uint32* OutSubpass
 			for (const FHlslVectorType& FragDeclType : FragDeclTypes)
 			{
 				// Try to find "Texture2D<T>" or "Texture2D< T >" (where T is the vector type), because a rewritten HLSL might have changed the formatting.
-				std::string LastFragDataN = GSpecialTextureLastFragData + FragDeclType.TypenameSuffix + "_" + std::to_string(SubpassIndex);
-				std::string FragDecl = "Texture2D<" + FragDeclType.TypenameIdent + "> " + LastFragDataN + ";";
-				size_t FragDeclIncludePos = SourceData.find(FragDecl);
+				FAnsiString LastFragDataN = FAnsiString::Printf("%s%s_%u", GSpecialTextureLastFragData, FragDeclType.TypenameSuffix, SubpassIndex);
+				FAnsiString FragDecl = FAnsiString::Printf("Texture2D<%s> %s;", FragDeclType.TypenameIdent, *LastFragDataN);
+				int32 FragDeclIncludePos = SourceData.Find(FragDecl, ESearchCase::CaseSensitive);
 			
-				if (FragDeclIncludePos == std::string::npos)
+				if (FragDeclIncludePos == INDEX_NONE)
 				{
-					FragDecl = "Texture2D< " + FragDeclType.TypenameIdent + " > " + LastFragDataN + ";";
-					FragDeclIncludePos = SourceData.find(FragDecl);
+					FragDecl = FAnsiString::Printf("Texture2D< %s > %s;", FragDeclType.TypenameIdent, *LastFragDataN);
+					FragDeclIncludePos = SourceData.Find(FragDecl, ESearchCase::CaseSensitive);
 				}
 			
-				if (FragDeclIncludePos != std::string::npos)
+				if (FragDeclIncludePos != INDEX_NONE)
 				{
 					// Replace declaration of Texture2D<T> with SubpassInput<T>
-					SourceData.replace(
+					ReplaceString(SourceData,
 						FragDeclIncludePos,
-						FragDecl.length(),
-						("[[vk::input_attachment_index(" + std::to_string(SubpassIndex) + ")]] SubpassInput<" + FragDeclType.TypenameIdent + "> " + LastFragDataN + ";")
-					);
+						FragDecl.Len(),
+						FAnsiString::Printf("[[vk::input_attachment_index(%d)]] SubpassInput<%s> %s;", SubpassIndex, FragDeclType.TypenameIdent, *LastFragDataN));
 
 					OutSubpassInputsDim[SubpassIndex] = FragDeclType.Dimension;
 
 					// Replace all uses of special texture by 'SubpassLoad' operation
-					std::string FragLoad = LastFragDataN + ".Load(uint3(0, 0, 0), 0)";
-					for (size_t FragLoadIncludePos = 0; (FragLoadIncludePos = SourceData.find(FragLoad, FragLoadIncludePos)) != std::string::npos;)
+					FAnsiString FragLoad = FAnsiString::Printf("%s.Load(uint3(0, 0, 0), 0)", *LastFragDataN);
+					for (int32 FragLoadIncludePos = 0; (FragLoadIncludePos = SourceData.Find(FragLoad, ESearchCase::CaseSensitive, ESearchDir::FromStart, FragLoadIncludePos)) != INDEX_NONE;)
 					{
-						SourceData.replace(
+						ReplaceString(SourceData,
 							FragLoadIncludePos,
-							FragLoad.length(),
-							(LastFragDataN + ".SubpassLoad()")
-						);
+							FragLoad.Len(),
+							FAnsiString::Printf("%s.SubpassLoad()", *LastFragDataN));
 					}
 
 					// Mark source data as being modified
diff --git a/Engine/Source/Developer/Apple/MetalShaderFormat/Private/MetalShaderCompiler.h b/Engine/Source/Developer/Apple/MetalShaderFormat/Private/MetalShaderCompiler.h
--- a/Engine/Source/Developer/Apple/MetalShaderFormat/Private/MetalShaderCompiler.h
+++ b/Engine/Source/Developer/Apple/MetalShaderFormat/Private/MetalShaderCompiler.h
@@ -101,4 +101,4 @@ struct FMetalShaderOutputMetaData
 };
 
 // Replace the special texture "gl_LastFragData" to a native subpass fetch operation. Returns true if the input source has been modified.
-extern bool PatchSpecialTextureInHlslSource(std::string& SourceData, uint32* OutSubpassInputsDim, uint32 SubpassInputDimCount);
+bool PatchSpecialTextureInHlslSource(FAnsiString& SourceData, uint32* OutSubpassInputsDim, uint32 SubpassInputDimCount);
