Moved Metal compiler temp files into UE managed folder to prevent multiple instances of ShaderCompiler clashing on the same files.

diff --git a/Engine/Source/Developer/Apple/MetalShaderFormat/Private/MetalShaderCompiler.cpp b/Engine/Source/Developer/Apple/MetalShaderFormat/Private/MetalShaderCompiler.cpp
--- a/Engine/Source/Developer/Apple/MetalShaderFormat/Private/MetalShaderCompiler.cpp
+++ b/Engine/Source/Developer/Apple/MetalShaderFormat/Private/MetalShaderCompiler.cpp
@@ -764,7 +764,13 @@ void BuildMetalShaderOutput(
 			Job.bOptimizeForSize = ShaderInput.Environment.GetCompileArgument(TEXT("METAL_OPTIMIZE_FOR_SIZE"), false);
 			Job.bCompileAsPCH = false;
 			Job.ReturnCode = 0;
-
+			
+			// Module cache path to ensure that we don't have clashes across multiple instances.
+			FString InputHashString;
+			BytesToHex(ShaderInput.Hash.GetBytes(), 8, InputHashString);
+			const FString& ModulesCachePath = TempDir / TEXT("ModuleCache") / InputHashString;
+			Job.ModuleCacheDirectory = ModulesCachePath;
+			
 			bSucceeded = FMetalCompilerToolchain::Get()->CompileMetalShader(Job, Bytecode);
 			if (bSucceeded)
 			{
diff --git a/Engine/Source/Developer/Apple/MetalShaderFormat/Private/MetalShaderCompiler.h b/Engine/Source/Developer/Apple/MetalShaderFormat/Private/MetalShaderCompiler.h
--- a/Engine/Source/Developer/Apple/MetalShaderFormat/Private/MetalShaderCompiler.h
+++ b/Engine/Source/Developer/Apple/MetalShaderFormat/Private/MetalShaderCompiler.h
@@ -54,6 +54,7 @@ struct FMetalShaderBytecodeJob
 	FString MathMode;
 	FString Standard;
 	FString IncludeDir;
+	FString ModuleCacheDirectory;
 	uint32 SourceCRCLen;
 	uint32 SourceCRC;
 	bool bRetainObjectFile;
diff --git a/Engine/Source/Developer/Apple/MetalShaderFormat/Private/MetalShaderFormat.cpp b/Engine/Source/Developer/Apple/MetalShaderFormat/Private/MetalShaderFormat.cpp
--- a/Engine/Source/Developer/Apple/MetalShaderFormat/Private/MetalShaderFormat.cpp
+++ b/Engine/Source/Developer/Apple/MetalShaderFormat/Private/MetalShaderFormat.cpp
@@ -1120,7 +1120,8 @@ bool FMetalCompilerToolchain::CompileMetalShader(FMetalShaderBytecodeJob& Job, F
 	FString IncludeArgs = Job.IncludeDir.Len() ? FString::Printf(TEXT("-I %s"), *Job.IncludeDir) : TEXT("");
 	{
 		// Invoke the metal frontend.
-		FString MetalParams = FString::Printf(TEXT("%s %s %s %s -Wno-null-character -fbracket-depth=1024 %s %s %s %s %s -o %s"), *Job.MinOSVersion, *Job.PreserveInvariance, *Job.DebugInfo, *Job.MathMode, TEXT("-c"), *Job.Standard, *Job.Defines, *IncludeArgs, *LocalInputMetalFilePath, *LocalOutputMetalAIRFilePath);
+		FString MetalParams = FString::Printf(TEXT("%s %s %s %s -Wno-null-character -fbracket-depth=1024 %s %s %s %s %s -fmodules-cache-path=%s -o %s"), *Job.MinOSVersion, *Job.PreserveInvariance, *Job.DebugInfo, *Job.MathMode, TEXT("-c"), *Job.Standard, *Job.Defines, *IncludeArgs, *LocalInputMetalFilePath, *Job.ModuleCacheDirectory, *LocalOutputMetalAIRFilePath);
+		
 		if (Job.bOptimizeForSize)
 		{
 			MetalParams += TEXT(" -Os");
@@ -1128,6 +1129,9 @@ bool FMetalCompilerToolchain::CompileMetalShader(FMetalShaderBytecodeJob& Job, F
 
 		bool bSuccess = this->ExecMetalFrontend(SDK, *MetalParams, &Job.ReturnCode, &Job.Results, &Job.Errors);
 
+		// Delete the cache directory because it takes a approx 5mb per compile
+		IFileManager::Get().DeleteDirectory(*Job.ModuleCacheDirectory, false, true);
+		
 		if (!bSuccess || (Job.ReturnCode != 0))
 		{
 			Job.Message = FString::Printf(TEXT("Failed to compile %s to bytecode %s, code: %d, output: %s %s"), *LocalInputMetalFilePath, *LocalOutputMetalAIRFilePath, Job.ReturnCode, *Job.Results, *Job.Errors);
