Added instruction count for SM6

diff --git a/Engine/Shaders/Public/Platform/Metal/MetalCommon.ush b/Engine/Shaders/Public/Platform/Metal/MetalCommon.ush
--- a/Engine/Shaders/Public/Platform/Metal/MetalCommon.ush
+++ b/Engine/Shaders/Public/Platform/Metal/MetalCommon.ush
@@ -7,7 +7,7 @@
 #pragma once
 
 // Update this GUID to invalidate shader recompilation for only Metal shaders
-#pragma message("UESHADERMETADATA_VERSION 1271B4C0-CB27-43E9-BD07-8FFC49447346")
+#pragma message("UESHADERMETADATA_VERSION F33F36E3-89FA-4A80-8B62-FB0B726F39E1")
 
 #if VERTEXSHADER && HULLSHADER
 #ifndef TESSELLATIONSHADER
diff --git a/Engine/Source/Developer/Apple/MetalShaderFormat/Private/MetalCompileShaderMSC.cpp b/Engine/Source/Developer/Apple/MetalShaderFormat/Private/MetalCompileShaderMSC.cpp
--- a/Engine/Source/Developer/Apple/MetalShaderFormat/Private/MetalCompileShaderMSC.cpp
+++ b/Engine/Source/Developer/Apple/MetalShaderFormat/Private/MetalCompileShaderMSC.cpp
@@ -317,7 +317,7 @@ static void ProcessReflection(ID3D12ShaderReflection* ShaderReflection, const ui
 	check(NumSRVs == 0 && NumUAVs == 0 && NumSamplers == 0);
 }
 
-static bool ReflectDXILAndBuildDescriptorRanges(const TArray<uint32>& DXILReflection, const FShaderCompilerInput& Input, FShaderCompilerOutput& Output, CrossCompiler::FHlslccHeaderWriter& CCHeaderWriter, FMetalShaderOutputMetaData& OutputData, uint32& NumCBVs)
+static bool ReflectDXILAndBuildDescriptorRanges(const TArray<uint32>& DXILReflection, const FShaderCompilerInput& Input, FShaderCompilerOutput& Output, CrossCompiler::FHlslccHeaderWriter& CCHeaderWriter, FMetalShaderOutputMetaData& OutputData, uint32& NumCBVs, uint32& NumInstructions)
 {
 	// Reflect DXIL
 	TRefCountPtr<IDxcUtils> Utils;
@@ -347,6 +347,13 @@ static bool ReflectDXILAndBuildDescriptorRanges(const TArray<uint32>& DXILReflec
 
 		D3D12_SHADER_DESC ShaderDesc = {};
 		ShaderReflection->GetDesc(&ShaderDesc);
+		
+		NumInstructions = ShaderDesc.InstructionCount;
+		
+		// Return a fraction of the number of instructions as DXIL is more verbose than DXBC.
+		// Ratio 119:307 was estimated by gathering average instruction count for D3D11 and D3D12 shaders in ShooterGame with result being ~ 357:921.
+		constexpr uint32 DxbcToDxilInstructionRatio[2] = { 119, 307 };
+		NumInstructions = NumInstructions * DxbcToDxilInstructionRatio[0] / DxbcToDxilInstructionRatio[1];
 
 		ProcessReflection(ShaderReflection.GetReference(), ShaderDesc.BoundResources, Input, Output, CCHeaderWriter, OutputData, NumCBVs);
 
@@ -589,7 +596,7 @@ void FMetalCompileShaderMSC::DoCompileMetalShader(
 		FCStringAnsi::Snprintf(MainCRC, 25, "Main_%0.8x_%0.8x", CRCLen, CRC);
 		
 		// Build shader metadata and root signature parameters
-		bool bSuccessfulReflection = ReflectDXILAndBuildDescriptorRanges(DxilData, Input, Output, CCHeaderWriter, OutputData, NumCBVs);
+		bool bSuccessfulReflection = ReflectDXILAndBuildDescriptorRanges(DxilData, Input, Output, CCHeaderWriter, OutputData, NumCBVs, Output.NumInstructions);
 		check(bSuccessfulReflection);
 		
 		// Build root parameters
diff --git a/Engine/Source/Developer/Apple/MetalShaderFormat/Private/MetalShaderCompiler.cpp b/Engine/Source/Developer/Apple/MetalShaderFormat/Private/MetalShaderCompiler.cpp
--- a/Engine/Source/Developer/Apple/MetalShaderFormat/Private/MetalShaderCompiler.cpp
+++ b/Engine/Source/Developer/Apple/MetalShaderFormat/Private/MetalShaderCompiler.cpp
@@ -635,7 +635,11 @@ void BuildMetalShaderOutput(
 		Ar.Serialize((void*)USFSource, SourceLen + 1 - (USFSource - InShaderSource));
 		
 		ShaderOutput.ModifiedShaderSource = MetalCode;
-		ShaderOutput.NumInstructions = NumLines;
+		
+		if(!bBindlessEnabled)
+		{
+			ShaderOutput.NumInstructions = NumLines;
+		}
 		ShaderOutput.NumTextureSamplers = Header.Bindings.NumSamplers;
 		ShaderOutput.bSucceeded = true;
 	}
@@ -836,7 +840,10 @@ void BuildMetalShaderOutput(
 		}
 
 		ShaderOutput.ModifiedShaderSource = MetalCode;
-		ShaderOutput.NumInstructions = NumLines;
+		if(!bBindlessEnabled)
+		{
+			ShaderOutput.NumInstructions = NumLines;
+		}
 		ShaderOutput.bSucceeded = bSucceeded;
 	}
 }
