Reduce the maximum number of shaders per library for SM6 temporarily for build machines

diff --git a/Engine/Source/Developer/Apple/MetalShaderFormat/Private/MetalShaderFormat.cpp b/Engine/Source/Developer/Apple/MetalShaderFormat/Private/MetalShaderFormat.cpp
--- a/Engine/Source/Developer/Apple/MetalShaderFormat/Private/MetalShaderFormat.cpp
+++ b/Engine/Source/Developer/Apple/MetalShaderFormat/Private/MetalShaderFormat.cpp
@@ -193,7 +193,7 @@ public:
 		// SM6 needs a lower limit of shaders per library as the packing process takes a significant amount of RAM
 		if(ShaderFormatName == NAME_SF_METAL_SM6)
 		{
-			NumShadersPerLibrary = 2000;
+			NumShadersPerLibrary = 1000;
 		}
 		
 		const FString ArchivePath = (WorkingDirectory / ShaderFormatAndShaderPlatformName.GetPlainNameString());
