MetalShaderConverter - Enable IRCompatibilityFlagSampleNanToZero by default to match behaviour of other IHMs

#tests CitySample with SM6 and MovieRenderQueue without the fix that avoid NaN in Texture3DSample

diff --git a/Engine/Shaders/Public/Platform/Metal/MetalCommon.ush b/Engine/Shaders/Public/Platform/Metal/MetalCommon.ush
--- a/Engine/Shaders/Public/Platform/Metal/MetalCommon.ush
+++ b/Engine/Shaders/Public/Platform/Metal/MetalCommon.ush
@@ -7,7 +7,7 @@
 #pragma once
 
 // Update this GUID to invalidate shader recompilation for only Metal shaders
-#pragma message("UESHADERMETADATA_VERSION 51B1BE79-9E82-469F-9F73-A42ADE28C5B6")
+#pragma message("UESHADERMETADATA_VERSION B5179016-7307-4B5B-899D-2E15BEDFA58B")
 
 #if VERTEXSHADER && HULLSHADER
 #ifndef TESSELLATIONSHADER
diff --git a/Engine/Source/Developer/Apple/MetalShaderFormat/Private/MetalCompileShaderMSC.cpp b/Engine/Source/Developer/Apple/MetalShaderFormat/Private/MetalCompileShaderMSC.cpp
--- a/Engine/Source/Developer/Apple/MetalShaderFormat/Private/MetalCompileShaderMSC.cpp
+++ b/Engine/Source/Developer/Apple/MetalShaderFormat/Private/MetalCompileShaderMSC.cpp
@@ -657,7 +657,7 @@ void FMetalCompileShaderMSC::DoCompileMetalShader(
 		IRCompilerSetEntryPointName(CompilerInstance, MainCRC);
 		IRCompilerSetGlobalRootSignature(CompilerInstance, RootSignature);
 		IRCompilerSetStageInGenerationMode(CompilerInstance, IRStageInCodeGenerationModeUseSeparateStageInFunction);
-		IRCompilerSetCompatibilityFlags(CompilerInstance, IRCompatibilityFlagBoundsCheck);
+		IRCompilerSetCompatibilityFlags(CompilerInstance, (IRCompatibilityFlags)(IRCompatibilityFlagBoundsCheck | IRCompatibilityFlagSampleNanToZero));
 		IRCompilerSetMinimumGPUFamily(CompilerInstance, IRGPUFamilyMetal3);
 		IRCompilerSetMinimumDeploymentTarget(CompilerInstance, IROperatingSystem_macOS, "15.0.0");
 		
