[Shaders] Fix crash in ParseCompilerVersionAndTarget() of Metal backend.

#rnx

diff --git a/Engine/Source/Developer/Apple/MetalShaderFormat/Private/MetalShaderFormat.cpp b/Engine/Source/Developer/Apple/MetalShaderFormat/Private/MetalShaderFormat.cpp
--- a/Engine/Source/Developer/Apple/MetalShaderFormat/Private/MetalShaderFormat.cpp
+++ b/Engine/Source/Developer/Apple/MetalShaderFormat/Private/MetalShaderFormat.cpp
@@ -483,7 +483,7 @@ public:
 
 IMPLEMENT_MODULE(FMetalShaderFormatModule, MetalShaderFormat);
 
-static FMetalCompilerToolchain::EMetalToolchainStatus ParseCompilerVersionAndTarget(const FString& OutputOfMetalDashV, FString& VersionString, FMetalCompilerToolchain::PackedVersion& PackedVersionNumber, FMetalCompilerToolchain::PackedVersion& PackedTargetNumber)
+static FMetalCompilerToolchain::EMetalToolchainStatus ParseCompilerVersionAndTarget(const FString& OutputOfMetalDashV, FString& OutVersionString, FMetalCompilerToolchain::PackedVersion& OutPackedVersionNumber, FMetalCompilerToolchain::PackedVersion& OutPackedTargetNumber)
 {
 	/*
 		Output of metal -v might look like this:
@@ -493,20 +493,27 @@ static FMetalCompilerToolchain::EMetalToolchainStatus ParseCompilerVersionAndTar
 		InstalledDir: C:\Program Files\Metal Developer Tools\ios\bin
 	*/
 
+	// Default initialize output parameters
+	OutVersionString.Empty();
+	OutPackedVersionNumber = {};
+	OutPackedTargetNumber = {};
+
 	TArray<FString> Lines;
 	OutputOfMetalDashV.ParseIntoArrayLines(Lines);
 	int32 VersionLineIndex = 0;
+
+	for (int32 Index = 0; Index < Lines.Num(); ++Index)
 	{
-		for (int32 Index = 0; Index < Lines.Num(); ++Index)
+		if (Lines[Index].StartsWith(TEXT("Apple ")) && Lines[Index].Contains(TEXT(" version ")) && Lines[Index].EndsWith(TEXT(")")))
 		{
-			if (Lines[Index].StartsWith(TEXT("Apple ")) && Lines[Index].Contains(TEXT(" version ")) && Lines[Index].EndsWith(TEXT(")")))
-			{
-				VersionLineIndex = Index;
-				break;
-			}
+			VersionLineIndex = Index;
+			break;
 		}
+	}
 
-		VersionString = Lines[VersionLineIndex];
+	if (VersionLineIndex < Lines.Num())
+	{
+		OutVersionString = Lines[VersionLineIndex];
 		FString& Version = Lines[VersionLineIndex];
 		check(!Version.IsEmpty());
 
@@ -528,20 +535,21 @@ static FMetalCompilerToolchain::EMetalToolchainStatus ParseCompilerVersionAndTar
 			UE_LOG(LogMetalCompilerSetup, Warning, TEXT("Obtained: %s"), *Version);
 		}
 		
-		PackedVersionNumber.Major = Major;
-		PackedVersionNumber.Minor = Minor;
+		OutPackedVersionNumber.Major = Major;
+		OutPackedVersionNumber.Minor = Minor;
 		// The version name in brackets is too irregular to extract a useful patch version
 		// Sometimes (metalfe-31001.667.2), sometimes (metalfe-31001.643.2.1), sometimes (metalfe-31001.362-windows)
-		PackedVersionNumber.Patch = 0;
+		OutPackedVersionNumber.Patch = 0;
 	}
 
-	if (PackedVersionNumber.Version == 0)
+	if (OutPackedVersionNumber.Version == 0)
 	{
 		return FMetalCompilerToolchain::EMetalToolchainStatus::CouldNotParseCompilerVersion;
 	}
 
+	if (VersionLineIndex + 1 < Lines.Num())
 	{
-		FString& FormatVersion = Lines[VersionLineIndex + 1];
+		const FString& FormatVersion = Lines[VersionLineIndex + 1];
 		int32 Major = 0, Minor = 0, Patch = 0;
 		int32 NumResults = 0;
 #if !PLATFORM_WINDOWS
@@ -549,12 +557,12 @@ static FMetalCompilerToolchain::EMetalToolchainStatus ParseCompilerVersionAndTar
 #else
 		NumResults = swscanf_s(*FormatVersion, TEXT("Target: air64-apple-darwin%d.%d.%d"), &Major, &Minor, &Patch);
 #endif
-		PackedTargetNumber.Major = Major;
-		PackedTargetNumber.Minor = Minor;
-		PackedTargetNumber.Patch = Patch;
+		OutPackedTargetNumber.Major = Major;
+		OutPackedTargetNumber.Minor = Minor;
+		OutPackedTargetNumber.Patch = Patch;
 	}
 
-	if (PackedTargetNumber.Version == 0)
+	if (OutPackedTargetNumber.Version == 0)
 	{
 		return FMetalCompilerToolchain::EMetalToolchainStatus::CouldNotParseTargetVersion;
 	}
@@ -954,11 +962,8 @@ FMetalCompilerToolchain::EMetalToolchainStatus FMetalCompilerToolchain::DoWindow
 	bool bUseLocalMetalToolchain = FPaths::FileExists(MetalFrontendBinaryCommand[AppleSDKMac]) && FPaths::FileExists(MetalFrontendBinaryCommand[AppleSDKMobile]);
 	if (!bUseLocalMetalToolchain)
 	{
-		if (GCheckCompilerToolChainSetup > 0)
-		{
-			UE_LOG(LogMetalCompilerSetup, Display, TEXT("Searching for Metal toolchain, but it doesn't appear to be installed."));
-			UE_LOG(LogMetalCompilerSetup, Display, TEXT("Searched for %s and %s"), *MetalFrontendBinaryCommand[AppleSDKMac], *MetalFrontendBinaryCommand[AppleSDKMobile]);
-		}
+		UE_LOG(LogMetalCompilerSetup, Display, TEXT("Searching for Metal toolchain, but it doesn't appear to be installed."));
+		UE_LOG(LogMetalCompilerSetup, Display, TEXT("Searched for %s and %s"), *MetalFrontendBinaryCommand[AppleSDKMac], *MetalFrontendBinaryCommand[AppleSDKMobile]);
 		return EMetalToolchainStatus::ToolchainNotFound;
 	}
 
@@ -978,10 +983,7 @@ FMetalCompilerToolchain::EMetalToolchainStatus FMetalCompilerToolchain::DoWindow
         !FPaths::FileExists(AirPackBinaryCommand[AppleSDKMac]) ||
         !FPaths::FileExists(AirPackBinaryCommand[AppleSDKMobile]))
 	{
-		if (GCheckCompilerToolChainSetup > 0)
-		{
-			UE_LOG(LogMetalCompilerSetup, Warning, TEXT("Missing toolchain binaries."))
-		}
+		UE_LOG(LogMetalCompilerSetup, Warning, TEXT("Missing toolchain binaries."))
 		return EMetalToolchainStatus::ToolchainNotFound;
 	}
 
