[Metal] Re-introduce condition that guards logs about missing Metal toolchian via CVar "Metal.CheckCompilerToolChainSetup".

This was previously removed with CL 39801290, but only the crash fix was needed.
This is spamming the log, but can remain active if aforementioned CVar is enabled.

#rnx

diff --git a/Engine/Source/Developer/Apple/MetalShaderFormat/Private/MetalShaderFormat.cpp b/Engine/Source/Developer/Apple/MetalShaderFormat/Private/MetalShaderFormat.cpp
--- a/Engine/Source/Developer/Apple/MetalShaderFormat/Private/MetalShaderFormat.cpp
+++ b/Engine/Source/Developer/Apple/MetalShaderFormat/Private/MetalShaderFormat.cpp
@@ -962,8 +962,11 @@ FMetalCompilerToolchain::EMetalToolchainStatus FMetalCompilerToolchain::DoWindow
 	bool bUseLocalMetalToolchain = FPaths::FileExists(MetalFrontendBinaryCommand[AppleSDKMac]) && FPaths::FileExists(MetalFrontendBinaryCommand[AppleSDKMobile]);
 	if (!bUseLocalMetalToolchain)
 	{
-		UE_LOG(LogMetalCompilerSetup, Display, TEXT("Searching for Metal toolchain, but it doesn't appear to be installed."));
-		UE_LOG(LogMetalCompilerSetup, Display, TEXT("Searched for %s and %s"), *MetalFrontendBinaryCommand[AppleSDKMac], *MetalFrontendBinaryCommand[AppleSDKMobile]);
+		if (GCheckCompilerToolChainSetup > 0)
+		{
+			UE_LOG(LogMetalCompilerSetup, Display, TEXT("Searching for Metal toolchain, but it doesn't appear to be installed."));
+			UE_LOG(LogMetalCompilerSetup, Display, TEXT("Searched for %s and %s"), *MetalFrontendBinaryCommand[AppleSDKMac], *MetalFrontendBinaryCommand[AppleSDKMobile]);
+		}
 		return EMetalToolchainStatus::ToolchainNotFound;
 	}
 
@@ -983,7 +986,10 @@ FMetalCompilerToolchain::EMetalToolchainStatus FMetalCompilerToolchain::DoWindow
         !FPaths::FileExists(AirPackBinaryCommand[AppleSDKMac]) ||
         !FPaths::FileExists(AirPackBinaryCommand[AppleSDKMobile]))
 	{
-		UE_LOG(LogMetalCompilerSetup, Warning, TEXT("Missing toolchain binaries."))
+		if (GCheckCompilerToolChainSetup > 0)
+		{
+			UE_LOG(LogMetalCompilerSetup, Warning, TEXT("Missing toolchain binaries."))
+		}
 		return EMetalToolchainStatus::ToolchainNotFound;
 	}
 
