[ShaderAnalyzer] Improve shader output artifacts for ShaderAnalyzer.

- Remove entry point name from output artifact name; This is already visible in the ShaderAnalyzer toolbar.
- Output Metal source as shader artifact in Metal/SPIR-V backend.
- Output DXIL disassembly as shader artifact in Metal/MSC backend.

#rnx

diff --git a/Engine/Source/Developer/Apple/MetalShaderFormat/Private/MetalCompileShaderMSC.cpp b/Engine/Source/Developer/Apple/MetalShaderFormat/Private/MetalCompileShaderMSC.cpp
--- a/Engine/Source/Developer/Apple/MetalShaderFormat/Private/MetalCompileShaderMSC.cpp
+++ b/Engine/Source/Developer/Apple/MetalShaderFormat/Private/MetalCompileShaderMSC.cpp
@@ -597,6 +597,16 @@ void FMetalCompileShaderMSC::DoCompileMetalShader(
 			return;
 		}
 		
+		// Return code reflection if requested for shader analysis
+		if (Input.Environment.CompilerFlags.Contains(CFLAG_OutputAnalysisArtifacts))
+		{
+			FGenericShaderStat ShaderCodeReflection;
+			if (CrossCompiler::FShaderConductorContext::Disassemble(CrossCompiler::EShaderConductorIR::Dxil, DxilData.GetData(), DxilData.Num()*sizeof(uint32), ShaderCodeReflection))
+			{
+				Output.ShaderStatistics.Add(MoveTemp(ShaderCodeReflection));
+			}
+		}
+
 		if (bDumpDebugInfo)
 		{
 			DumpDebugShaderBinary(Input, DxilData.GetData(), DxilData.Num() * sizeof(uint32), TEXT("dxil"));
diff --git a/Engine/Source/Developer/Apple/MetalShaderFormat/Private/MetalCompileShaderSPIRV.cpp b/Engine/Source/Developer/Apple/MetalShaderFormat/Private/MetalCompileShaderSPIRV.cpp
--- a/Engine/Source/Developer/Apple/MetalShaderFormat/Private/MetalCompileShaderSPIRV.cpp
+++ b/Engine/Source/Developer/Apple/MetalShaderFormat/Private/MetalCompileShaderSPIRV.cpp
@@ -219,6 +219,7 @@ void FMetalCompileShaderSPIRV::DoCompileMetalShader(
 	FString PreprocessedShader = InPreprocessedShader;
 	bool bUsingInlineRayTracing = Input.Environment.CompilerFlags.Contains(CFLAG_InlineRayTracing);
 	const bool bBindlessEnabled = (Input.Environment.CompilerFlags.Contains(CFLAG_BindlessResources) || Input.Environment.CompilerFlags.Contains(CFLAG_BindlessSamplers));
+	const bool bOutputAnalysisArtifacts = Input.Environment.CompilerFlags.Contains(CFLAG_OutputAnalysisArtifacts);
 
 #if PLATFORM_MAC || PLATFORM_WINDOWS
 	{
@@ -279,7 +280,6 @@ void FMetalCompileShaderSPIRV::DoCompileMetalShader(
 		}
 		
 		// If source data was modified, reload it into the compiler context
-		const bool bOutputAnalysisArtifacts = Input.Environment.CompilerFlags.Contains(CFLAG_OutputAnalysisArtifacts);
 		if (bSourceDataWasModified)
 		{
 			CompilerContext.LoadSource(FAnsiStringView(*SourceData, SourceData.Len()), Input.VirtualSourceFilePath, Input.EntryPointName, Frequency);
@@ -1367,6 +1367,11 @@ void FMetalCompileShaderSPIRV::DoCompileMetalShader(
 		MetalSource.replace(Begin, ClangOptNoneString.length(), " ");
 	}
 	
+	if (bOutputAnalysisArtifacts)
+	{
+		Output.AddStatistic<FString>(TEXT("Metal"), ANSI_TO_TCHAR(&MetalSource[0]), FGenericShaderStat::EFlags::Hidden, FShaderStatTagNames::AnalysisArtifactsName);
+	}
+
 	if (bDumpDebugInfo && !MetalSource.empty())
 	{
 		DumpDebugShaderText(Input, &MetalSource[0], MetalSource.size(), TEXT("metal"));
diff --git a/Engine/Source/Developer/Windows/ShaderFormatD3D/Private/D3DShaderCompilerDXC.cpp b/Engine/Source/Developer/Windows/ShaderFormatD3D/Private/D3DShaderCompilerDXC.cpp
--- a/Engine/Source/Developer/Windows/ShaderFormatD3D/Private/D3DShaderCompilerDXC.cpp
+++ b/Engine/Source/Developer/Windows/ShaderFormatD3D/Private/D3DShaderCompilerDXC.cpp
@@ -1194,7 +1194,6 @@ bool CompileAndProcessD3DShaderDXC(
 				FGenericShaderStat ShaderCodeReflection;
 				if (CrossCompiler::FShaderConductorContext::Disassemble(CrossCompiler::EShaderConductorIR::Dxil, ShaderBlob->GetBufferPointer(), ShaderBlob->GetBufferSize(), ShaderCodeReflection))
 				{
-					ShaderCodeReflection.StatName = FName(FString::Printf(TEXT("%s (%s)"), *ShaderCodeReflection.StatName.ToString(), *Input.EntryPointName));
 					Output.ShaderStatistics.Add(MoveTemp(ShaderCodeReflection));
 				}
 			}
