Update MetalDeveloperTools  to version 4.4 from 4.1

diff --git a/Engine/Source/Developer/Apple/MetalShaderFormat/Private/MetalShaderFormat.cpp b/Engine/Source/Developer/Apple/MetalShaderFormat/Private/MetalShaderFormat.cpp
--- a/Engine/Source/Developer/Apple/MetalShaderFormat/Private/MetalShaderFormat.cpp
+++ b/Engine/Source/Developer/Apple/MetalShaderFormat/Private/MetalShaderFormat.cpp
@@ -630,7 +630,7 @@ FString FMetalCompilerToolchain::XcrunPath(TEXT("/usr/bin/xcrun"));
 FString FMetalCompilerToolchain::MetalMacSDK(TEXT("macosx"));
 FString FMetalCompilerToolchain::MetalMobileSDK(TEXT("iphoneos"));
 
-FString FMetalCompilerToolchain::WindowsToolchainVersion(TEXT("4.1"));
+FString FMetalCompilerToolchain::WindowsToolchainVersion(TEXT("4.4"));
 
 // Static methods
 
