Metal MSC - Apply recommended CompatibilityFlags to ensure greater compatibility with other RHIs

#tests Local performance testing & ensuring no corruption were present

diff --git a/Engine/Shaders/Public/Platform/Metal/MetalCommon.ush b/Engine/Shaders/Public/Platform/Metal/MetalCommon.ush
--- a/Engine/Shaders/Public/Platform/Metal/MetalCommon.ush
+++ b/Engine/Shaders/Public/Platform/Metal/MetalCommon.ush
@@ -7,7 +7,7 @@
 #pragma once
 
 // Update this GUID to invalidate shader recompilation for only Metal shaders
-#pragma message("UESHADERMETADATA_VERSION B5179016-7307-4B5B-899D-2E15BEDFA58B")
+#pragma message("UESHADERMETADATA_VERSION A68B2D92-5E85-40D6-89F0-FCB78A6EAC75")
 
 #if VERTEXSHADER && HULLSHADER
 #ifndef TESSELLATIONSHADER
diff --git a/Engine/Source/Developer/Apple/MetalShaderFormat/Private/MetalCompileShaderMSC.cpp b/Engine/Source/Developer/Apple/MetalShaderFormat/Private/MetalCompileShaderMSC.cpp
--- a/Engine/Source/Developer/Apple/MetalShaderFormat/Private/MetalCompileShaderMSC.cpp
+++ b/Engine/Source/Developer/Apple/MetalShaderFormat/Private/MetalCompileShaderMSC.cpp
@@ -667,7 +667,7 @@ void FMetalCompileShaderMSC::DoCompileMetalShader(
 		IRCompilerSetEntryPointName(CompilerInstance, MainCRC);
 		IRCompilerSetGlobalRootSignature(CompilerInstance, RootSignature);
 		IRCompilerSetStageInGenerationMode(CompilerInstance, IRStageInCodeGenerationModeUseSeparateStageInFunction);
-		IRCompilerSetCompatibilityFlags(CompilerInstance, (IRCompatibilityFlags)(IRCompatibilityFlagBoundsCheck | IRCompatibilityFlagSampleNanToZero));
+		IRCompilerSetCompatibilityFlags(CompilerInstance, (IRCompatibilityFlags)(IRCompatibilityFlagBoundsCheck | IRCompatibilityFlagPositionInvariance | IRCompatibilityFlagSampleNanToZero | IRCompatibilityFlagTexWriteRoundingRTZ));
 		IRCompilerSetMinimumGPUFamily(CompilerInstance, IRGPUFamilyMetal3);
 		IRCompilerSetMinimumDeploymentTarget(CompilerInstance, IROperatingSystem_macOS, "15.0.0");
 		
