Metal - Fix Warnings

#tests Local build on mac & windows

diff --git a/Engine/Source/Developer/Apple/MetalShaderFormat/Private/MetalCompileShaderMSC.cpp b/Engine/Source/Developer/Apple/MetalShaderFormat/Private/MetalCompileShaderMSC.cpp
--- a/Engine/Source/Developer/Apple/MetalShaderFormat/Private/MetalCompileShaderMSC.cpp
+++ b/Engine/Source/Developer/Apple/MetalShaderFormat/Private/MetalCompileShaderMSC.cpp
@@ -52,11 +52,14 @@ extern void BuildMetalShaderOutput(
 
 #include "ShaderConductorContext.h"
 
-#include "d3d12shader.h"
-#include "dxc/dxcapi.h"
-
 #if PLATFORM_WINDOWS
-#include <dxc/Support/dxcapi.use.h>
+#include "Windows/AllowWindowsPlatformTypes.h"
+#endif
+	#include "d3d12shader.h"
+	#include "dxc/dxcapi.h"
+#if PLATFORM_WINDOWS
+	#include <dxc/Support/dxcapi.use.h>
+#include "Windows/HideWindowsPlatformTypes.h"
 #endif
 
 inline IRShaderStage ShaderFrequencyToStage(const EShaderFrequency UEStage)
@@ -568,7 +571,6 @@ void FMetalCompileShaderMSC::DoCompileMetalShader(
 	uint32 CRCLen = 0;
 	uint32 CRC = 0;
 	uint32 SourceLen = 0;
-	int32 Result = 0;
 	
 	struct FMetalResourceTableEntry : FUniformResourceEntry
 	{
@@ -677,8 +679,6 @@ void FMetalCompileShaderMSC::DoCompileMetalShader(
 		
 		if (!CompilerContext.CompileHlslToDxil(Options, DxilData))
 		{
-			Result = 0;
-			
 			UE_LOG(LogShaders, Error, TEXT("Failed to produce DXIL bytecode for '%s' '%s'!"), *Input.EntryPointName, *Input.DumpDebugInfoPath);
 			CompilerContext.FlushErrors(Output.Errors);
 			
@@ -707,8 +707,6 @@ void FMetalCompileShaderMSC::DoCompileMetalShader(
 			DumpDebugShaderBinary(Input, DxilData.GetData(), DxilData.Num() * sizeof(uint32), TEXT("dxil"));
 		}
 		
-		Result = 1;
-		
 		ANSICHAR MainCRC[25];
 		CRCLen = DxilData.Num() * sizeof( uint32_t );
 		CRC = FCrc::MemCrc_DEPRECATED(DxilData.GetData(), CRCLen);
@@ -941,25 +939,10 @@ void FMetalCompileShaderMSC::DoCompileMetalShader(
 		DumpDebugShaderText(Input, &MetalSource[0], MetalSource.size(), TEXT("metal"));
 	}
 
-	if (Result != 0)
-	{
-		Output.Target = Input.Target;
-		BuildMetalShaderOutput(Output, Input, MetalSource.c_str(), MetalSource.length(), CRCLen, CRC, VersionEnum, *Standard, *MinOSVersion, Output.Errors, OutputData.TypedBuffers, OutputData.InvariantBuffers, OutputData.TypedUAVs, OutputData.ConstantBuffers, bAllowFastIntrinsics
-		  , NumCBVs, OutputSizeVS, MaxInputPrimitivesPerMeshThreadgroupGS, bUsesDiscard, ReflectionJSON, MetalBytecode
-		);
-	}
-	else
-	{
-		// Log errors on failed compilation in this backend only when compiling from a debug dump USF.
-		if (EnumHasAnyFlags(Input.DebugInfoFlags, EShaderDebugInfoFlags::CompileFromDebugUSF))
-		{
-			for (const FShaderCompilerError& Error : Output.Errors)
-			{
-				UE_LOG(LogShaders, Error, TEXT("%s"), *Error.GetErrorStringWithLineMarker());
-			}
-			GLog->Flush();
-		}
-	}
+	Output.Target = Input.Target;
+	BuildMetalShaderOutput(Output, Input, MetalSource.c_str(), MetalSource.length(), CRCLen, CRC, VersionEnum, *Standard, *MinOSVersion, Output.Errors, OutputData.TypedBuffers, OutputData.InvariantBuffers, OutputData.TypedUAVs, OutputData.ConstantBuffers, bAllowFastIntrinsics
+		, NumCBVs, OutputSizeVS, MaxInputPrimitivesPerMeshThreadgroupGS, bUsesDiscard, ReflectionJSON, MetalBytecode
+	);
 }
 
 #endif // PLATFORM_MAC || PLATFORM_WINDOWS
diff --git a/Engine/Source/Developer/Apple/MetalShaderFormat/Private/MetalShaderCompiler.cpp b/Engine/Source/Developer/Apple/MetalShaderFormat/Private/MetalShaderCompiler.cpp
--- a/Engine/Source/Developer/Apple/MetalShaderFormat/Private/MetalShaderCompiler.cpp
+++ b/Engine/Source/Developer/Apple/MetalShaderFormat/Private/MetalShaderCompiler.cpp
@@ -25,7 +25,7 @@
 #include "MetalCompileShaderSPIRV.h"
 #include "MetalCompileShaderMSC.h"
 
-#if PLATFORM_MAC
+#if PLATFORM_MAC || PLATFORM_WINDOWS
 THIRD_PARTY_INCLUDES_START
 #include "metal_irconverter.h"
 THIRD_PARTY_INCLUDES_END
@@ -34,7 +34,6 @@ THIRD_PARTY_INCLUDES_END
 #if PLATFORM_WINDOWS
 #include "Windows/AllowWindowsPlatformTypes.h"
 THIRD_PARTY_INCLUDES_START
-	#include "metal_irconverter.h"
 	#include <objbase.h>
 	#include <assert.h>
 	#include <stdio.h>
