Metal - Cleanup intermediate files from MetalShaderCompilation after being use

#tests Local iOS build & preflight for ES3_1, SM5 & SM6

diff --git a/Engine/Source/Developer/Apple/MetalShaderFormat/Private/MetalShaderCompiler.cpp b/Engine/Source/Developer/Apple/MetalShaderFormat/Private/MetalShaderCompiler.cpp
--- a/Engine/Source/Developer/Apple/MetalShaderFormat/Private/MetalShaderCompiler.cpp
+++ b/Engine/Source/Developer/Apple/MetalShaderFormat/Private/MetalShaderCompiler.cpp
@@ -800,6 +800,8 @@ void BuildMetalShaderOutput(
 				Error->ErrorLineString = TEXT("0");
 				Error->StrippedErrorMessage = Job.Message;
 			}
+			// Clean up the .metal file after building it
+			IFileManager::Get().Delete(*MetalFileName);
 		}
 		
 		if (bSucceeded)
diff --git a/Engine/Source/Developer/Apple/MetalShaderFormat/Private/MetalShaderFormat.cpp b/Engine/Source/Developer/Apple/MetalShaderFormat/Private/MetalShaderFormat.cpp
--- a/Engine/Source/Developer/Apple/MetalShaderFormat/Private/MetalShaderFormat.cpp
+++ b/Engine/Source/Developer/Apple/MetalShaderFormat/Private/MetalShaderFormat.cpp
@@ -1173,21 +1173,33 @@ bool FMetalCompilerToolchain::CompileMetalShader(FMetalShaderBytecodeJob& Job, F
 	{
 		// Retain the .air. This usually means we are using shared native libraries.
 		bool bSuccess = FFileHelper::LoadFileToArray(Output.ObjectFile, *LocalOutputMetalAIRFilePath);
+		
+		// Clean up the .air file
+		IFileManager::Get().Delete(*LocalOutputMetalAIRFilePath);
+		
 		if (!bSuccess)
 		{
 			Job.Message = FString::Printf(TEXT("Failed to store AIR %s"), *LocalOutputMetalAIRFilePath);
 			return false;
 		}
 	}
+	else
+	{
+		// Clean up the .air file
+		IFileManager::Get().Delete(*LocalOutputMetalAIRFilePath);
+	}
 
 	{
 		// Retain the .metallib
 		Output.NativePath = LocalInputMetalFilePath;
 		bool bSuccess = FFileHelper::LoadFileToArray(Output.OutputFile, *LocalOutputMetalLibFilePath);
 
+		// Clean up the .metallib file
+		IFileManager::Get().Delete(*LocalOutputMetalLibFilePath);
+		
 		if (!bSuccess)
 		{
-			Job.Message = FString::Printf(TEXT("Failed to store metallib %s"), *LocalOutputMetalAIRFilePath);
+			Job.Message = FString::Printf(TEXT("Failed to store metallib %s"), *LocalOutputMetalLibFilePath);
 			return false;
 		}
 	}
