Metal - Revert MetalDeveloperTools back to 4.4 from 5.0-beta2

#tests preflight

diff --git a/Engine/Source/Developer/Apple/MetalShaderFormat/Private/MetalShaderFormat.cpp b/Engine/Source/Developer/Apple/MetalShaderFormat/Private/MetalShaderFormat.cpp
--- a/Engine/Source/Developer/Apple/MetalShaderFormat/Private/MetalShaderFormat.cpp
+++ b/Engine/Source/Developer/Apple/MetalShaderFormat/Private/MetalShaderFormat.cpp
@@ -619,8 +619,8 @@ FString FMetalCompilerToolchain::XcrunPath(TEXT("/usr/bin/xcrun"));
 FString FMetalCompilerToolchain::MetalMacSDK(TEXT("macosx"));
 FString FMetalCompilerToolchain::MetalMobileSDK(TEXT("iphoneos"));
 
-FString FMetalCompilerToolchain::WindowsToolchainVersion(TEXT("5.0-beta2"));
-FString FMetalCompilerToolchain::WindowsToolchainSubversion(TEXT("32023"));
+FString FMetalCompilerToolchain::WindowsToolchainVersion(TEXT("4.4"));
+FString FMetalCompilerToolchain::WindowsToolchainSubversion(TEXT(""));
 
 // Static methods
 
