Metal - Upgrade Metal Developer Tools from 4.4 to 5.3

#tests Local testing on Mac & iOS

diff --git a/Engine/Shaders/Public/Platform/Metal/MetalCommon.ush b/Engine/Shaders/Public/Platform/Metal/MetalCommon.ush
--- a/Engine/Shaders/Public/Platform/Metal/MetalCommon.ush
+++ b/Engine/Shaders/Public/Platform/Metal/MetalCommon.ush
@@ -7,7 +7,7 @@
 #pragma once
 
 // Update this GUID to invalidate shader recompilation for only Metal shaders
-#pragma message("UESHADERMETADATA_VERSION 92F6351B-26F7-43BE-A6B7-DA4B8E990D78")
+#pragma message("UESHADERMETADATA_VERSION 99D4B0D2-064D-4B5C-9292-57A1ED268BB9")
 
 #if VERTEXSHADER && HULLSHADER
 #ifndef TESSELLATIONSHADER
diff --git a/Engine/Source/Developer/Apple/MetalShaderFormat/Private/MetalShaderFormat.cpp b/Engine/Source/Developer/Apple/MetalShaderFormat/Private/MetalShaderFormat.cpp
--- a/Engine/Source/Developer/Apple/MetalShaderFormat/Private/MetalShaderFormat.cpp
+++ b/Engine/Source/Developer/Apple/MetalShaderFormat/Private/MetalShaderFormat.cpp
@@ -611,8 +611,8 @@ FString FMetalCompilerToolchain::XcrunPath(TEXT("/usr/bin/xcrun"));
 FString FMetalCompilerToolchain::MetalMacSDK(TEXT("macosx"));
 FString FMetalCompilerToolchain::MetalMobileSDK(TEXT("iphoneos"));
 
-FString FMetalCompilerToolchain::WindowsToolchainVersion(TEXT("4.4"));
-FString FMetalCompilerToolchain::WindowsToolchainSubversion(TEXT(""));
+FString FMetalCompilerToolchain::WindowsToolchainVersion(TEXT("5.3"));
+FString FMetalCompilerToolchain::WindowsToolchainSubversion(TEXT("32023"));
 
 // Static methods
 
