Metal Bindless - Fix bad assumption causing archiving issues & reenable optional bindless for SM5

#tests local cook & preflight

diff --git a/Engine/Config/Mac/DataDrivenPlatformInfo.ini b/Engine/Config/Mac/DataDrivenPlatformInfo.ini
--- a/Engine/Config/Mac/DataDrivenPlatformInfo.ini
+++ b/Engine/Config/Mac/DataDrivenPlatformInfo.ini
@@ -61,7 +61,7 @@ bSupportsUInt64ImageAtomics=false
 bSupportsNNEShaders=false
 bSupportsIndependentSamplers=true
 bSupportsComputeFramework = true
-bSupportsBindless=false
+bSupportsBindless=true
 FriendlyName=LOCTEXT("FriendlyShaderPlatformName_MacSM5", "Mac Metal SM5")
 
 [ShaderPlatform METAL_SM6]
@@ -122,7 +122,7 @@ bSupportsShaderPipelines = false
 bSupportsDualSourceBlending = true
 bSupportsIndependentSamplers=true
 bSupportsComputeFramework = true
-bSupportsBindless=false
+bSupportsBindless=true
 FriendlyName=LOCTEXT("FriendlyShaderPlatformName_MacMobile", "Mac Metal Mobile")
 
 [PreviewPlatform METAL_SM5]
diff --git a/Engine/Source/Developer/Apple/MetalShaderFormat/Private/MetalShaderCompiler.cpp b/Engine/Source/Developer/Apple/MetalShaderFormat/Private/MetalShaderCompiler.cpp
--- a/Engine/Source/Developer/Apple/MetalShaderFormat/Private/MetalShaderCompiler.cpp
+++ b/Engine/Source/Developer/Apple/MetalShaderFormat/Private/MetalShaderCompiler.cpp
@@ -1335,7 +1335,8 @@ bool FinalizeLibrary_Metal(FName const& Format, FString const& WorkingDir, FStri
 	const bool bCompilerAvailable = Toolchain->IsCompilerAvailable();
 	EShaderPlatform Platform = FMetalCompilerToolchain::MetalShaderFormatToLegacyShaderPlatform(Format);
 	const EAppleSDKType SDK = FMetalCompilerToolchain::MetalShaderPlatformToSDK(Platform);
-    bool bCompiledWithMetalShaderConverter = FDataDrivenShaderPlatformInfo::GetSupportsBindless(Platform);
+    bool bCompiledWithMetalShaderConverter = FDataDrivenShaderPlatformInfo::GetSupportsBindless(Platform) &&
+											 UE::ShaderCompiler::GetBindlessConfiguration(Platform) != ERHIBindlessConfiguration::Disabled;
     
 	if (bCompilerAvailable)
 	{
