Metal - Rename air-pack to metal-pack which is the current official name and the only one available in Xcode 26 Metal Toolchain whereas metal-pack existed previously

#tests Local cook & preflight Lyra Mac

diff --git a/Engine/Source/Developer/Apple/MetalShaderFormat/Private/MetalShaderCompiler.cpp b/Engine/Source/Developer/Apple/MetalShaderFormat/Private/MetalShaderCompiler.cpp
--- a/Engine/Source/Developer/Apple/MetalShaderFormat/Private/MetalShaderCompiler.cpp
+++ b/Engine/Source/Developer/Apple/MetalShaderFormat/Private/MetalShaderCompiler.cpp
@@ -1361,8 +1361,8 @@ bool FinalizeLibrary_Metal(FName const& Format, FString const& WorkingDir, FStri
             // Number of air per batch (limited by PlatformProcessLimits::MaxArgvParameters).
             static constexpr int32 NumArgcPerBatch = 96;
     
-            TArray<FString> AirPackArgsBatches;
-            FString AirPackArgs;
+            TArray<FString> MetalPackArgsBatches;
+            FString MetalPackArgs;
             uint32 CurPackArgsArgc = 0;
     
             uint32 Index = 0;
@@ -1377,21 +1377,21 @@ bool FinalizeLibrary_Metal(FName const& Format, FString const& WorkingDir, FStri
 
                 FString SourceFilePath = FString::Printf(TEXT("\"%s/%s\""), *FullyQualifiedWorkingDir, *FileName);
 
-                AirPackArgs += FString::Printf(TEXT("%s "), *SourceFilePath);
+				MetalPackArgs += FString::Printf(TEXT("%s "), *SourceFilePath);
                 CurPackArgsArgc++;
 
                 if (CurPackArgsArgc > NumArgcPerBatch)
                 {
-                    AirPackArgsBatches.Add(AirPackArgs);
-                    AirPackArgs.Empty(); // TODO: Might switch to SetNum(0); as Empty() implicitly reallocs (IIRC)
+					MetalPackArgsBatches.Add(MetalPackArgs);
+					MetalPackArgs.Empty(); // TODO: Might switch to SetNum(0); as Empty() implicitly reallocs (IIRC)
                     CurPackArgsArgc = 0;
                 }
             }
 
             // Add pending batch to the list.
-            if (AirPackArgs.Len() > 0)
+            if (MetalPackArgs.Len() > 0)
             {
-                AirPackArgsBatches.Add(AirPackArgs);
+				MetalPackArgsBatches.Add(MetalPackArgs);
             }
             bArchiveFileValid = (Shaders.Num() > 0);
 
@@ -1399,28 +1399,28 @@ bool FinalizeLibrary_Metal(FName const& Format, FString const& WorkingDir, FStri
                 // handle compile error
                 if (ReturnCode == 0 && bArchiveFileValid)
                 {
-                    // AirPack each batch.
+                    // MetalPack each batch.
                     FString BatchMergeArgs;
-                    for (int32 BatchIdx = 0; BatchIdx < AirPackArgsBatches.Num(); BatchIdx++)
+                    for (int32 BatchIdx = 0; BatchIdx < MetalPackArgsBatches.Num(); BatchIdx++)
                     {
-                        FString BatchOutputFile = FString::Printf(TEXT("AirPackBatch_%d_"), BatchIdx);
+                        FString BatchOutputFile = FString::Printf(TEXT("MetalPackBatch_%d_"), BatchIdx);
                         FString BatchOutputPath = FPaths::CreateTempFilename(*FullyQualifiedWorkingDir, *BatchOutputFile);
 
                         BatchMergeArgs += FString::Printf(TEXT("\"%s\" "), *BatchOutputPath);
 
-                        UE_LOG(LogMetalShaderCompiler, Display, TEXT("[%d/%d] %s"), (BatchIdx + 1), AirPackArgsBatches.Num(), *BatchOutputPath);
+                        UE_LOG(LogMetalShaderCompiler, Display, TEXT("[%d/%d] %s"), (BatchIdx + 1), MetalPackArgsBatches.Num(), *BatchOutputPath);
 
-                        FString AirPackParams = FString::Printf(TEXT("-pack-metallibs internal -pack-descriptors internal -pack-reflections internal %s -o \"%s\""), *AirPackArgsBatches[BatchIdx], *BatchOutputPath);
+                        FString MetalPackParams = FString::Printf(TEXT("-pack-metallibs internal -pack-descriptors internal -pack-reflections internal %s -o \"%s\""), *MetalPackArgsBatches[BatchIdx], *BatchOutputPath);
                         ReturnCode = 0;
                         Results.Empty();
                         Errors.Empty();
 
-                        bool bSuccess = Toolchain->ExecAirPack(SDK, *AirPackParams, &ReturnCode, &Results, &Errors);
+                        bool bSuccess = Toolchain->ExecMetalPack(SDK, *MetalPackParams, &ReturnCode, &Results, &Errors);
 
                         // handle compile error
                         if (!bSuccess || ReturnCode != 0)
                         {
-                            UE_LOG(LogShaders, Error, TEXT("Archiving failed: air-pack failed with code %d: %s %s"), ReturnCode, *Results, *Errors);
+                            UE_LOG(LogShaders, Error, TEXT("Archiving failed: metal-pack failed with code %d: %s %s"), ReturnCode, *Results, *Errors);
                         }
                     }
 
@@ -1434,12 +1434,12 @@ bool FinalizeLibrary_Metal(FName const& Format, FString const& WorkingDir, FStri
                         UE_LOG(LogMetalShaderCompiler, Warning, TEXT("Archiving warning: target metallib already exists and will be overwritten: %s"), *LocalMetalLibPath);
                     }
 
-                    FString AirPackParams = FString::Printf(TEXT("-pack-metallibs internal -pack-descriptors internal -pack-reflections internal %s -o \"%s\""), *BatchMergeArgs, *LocalMetalLibPath);
+                    FString MetalPackParams = FString::Printf(TEXT("-pack-metallibs internal -pack-descriptors internal -pack-reflections internal %s -o \"%s\""), *BatchMergeArgs, *LocalMetalLibPath);
                     ReturnCode = 0;
                     Results.Empty();
                     Errors.Empty();
 
-                    bool bSuccess = Toolchain->ExecAirPack(SDK, *AirPackParams, &ReturnCode, &Results, &Errors);
+                    bool bSuccess = Toolchain->ExecMetalPack(SDK, *MetalPackParams, &ReturnCode, &Results, &Errors);
 
                     // handle compile error
                     if (bSuccess && ReturnCode == 0)
@@ -1455,12 +1455,12 @@ bool FinalizeLibrary_Metal(FName const& Format, FString const& WorkingDir, FStri
                     }
                     else
                     {
-                        UE_LOG(LogShaders, Error, TEXT("Archiving failed: air-pack failed with code %d: %s %s"), ReturnCode, *Results, *Errors);
+                        UE_LOG(LogShaders, Error, TEXT("Archiving failed: metal-pack failed with code %d: %s %s"), ReturnCode, *Results, *Errors);
                     }
                 }
                 else
                 {
-                    UE_LOG(LogShaders, Error, TEXT("Archiving failed: no valid input for air-pack."));
+                    UE_LOG(LogShaders, Error, TEXT("Archiving failed: no valid input for metal-pack."));
                 }
             }
         }
diff --git a/Engine/Source/Developer/Apple/MetalShaderFormat/Private/MetalShaderFormat.cpp b/Engine/Source/Developer/Apple/MetalShaderFormat/Private/MetalShaderFormat.cpp
--- a/Engine/Source/Developer/Apple/MetalShaderFormat/Private/MetalShaderFormat.cpp
+++ b/Engine/Source/Developer/Apple/MetalShaderFormat/Private/MetalShaderFormat.cpp
@@ -599,12 +599,12 @@ FString FMetalCompilerToolchain::MetalObjectExtension(TEXT(".air"));
 FString FMetalCompilerToolchain::MetalFrontendBinary(TEXT("metal.exe"));
 FString FMetalCompilerToolchain::MetalArBinary(TEXT("metal-ar.exe"));
 FString FMetalCompilerToolchain::MetalLibraryBinary(TEXT("metallib.exe"));
-FString FMetalCompilerToolchain::AirPackBinary(TEXT("air-pack.exe"));
+FString FMetalCompilerToolchain::MetalPackBinary(TEXT("metal-pack.exe"));
 #else
 FString FMetalCompilerToolchain::MetalFrontendBinary(TEXT("metal"));
 FString FMetalCompilerToolchain::MetalArBinary(TEXT("metal-ar"));
 FString FMetalCompilerToolchain::MetalLibraryBinary(TEXT("metallib"));
-FString FMetalCompilerToolchain::AirPackBinary(TEXT("air-pack"));
+FString FMetalCompilerToolchain::MetalPackBinary(TEXT("metal-pack"));
 #endif
 
 FString FMetalCompilerToolchain::MetalMapExtension(TEXT(".metalmap"));
@@ -865,16 +865,16 @@ FMetalCompilerToolchain::EMetalToolchainStatus FMetalCompilerToolchain::DoMacNat
 			MetalLibBinaryCommand[AppleSDKMac] = ToolchainBase / TEXT("macos") / TEXT("bin") / MetalLibraryBinary;
 			MetalLibBinaryCommand[AppleSDKMobile] = ToolchainBase / TEXT("ios") / TEXT("bin") / MetalLibraryBinary;
 
-			AirPackBinaryCommand[AppleSDKMac] = ToolchainBase / TEXT("macos") / TEXT("bin") / AirPackBinary;
-            AirPackBinaryCommand[AppleSDKMobile] = ToolchainBase / TEXT("ios") / TEXT("bin") / AirPackBinary;
+			MetalPackBinaryCommand[AppleSDKMac] = ToolchainBase / TEXT("macos") / TEXT("bin") / MetalPackBinary;
+            MetalPackBinaryCommand[AppleSDKMobile] = ToolchainBase / TEXT("ios") / TEXT("bin") / MetalPackBinary;
 
 			if (!FPaths::FileExists(MetalArBinaryCommand[AppleSDKMac]) ||
 				!FPaths::FileExists(MetalArBinaryCommand[AppleSDKMobile]) ||
 				!FPaths::FileExists(MetalLibBinaryCommand[AppleSDKMac]) ||
 				!FPaths::FileExists(MetalLibBinaryCommand[AppleSDKMobile]) ||
 				!FPaths::FileExists(MetalLibBinaryCommand[AppleSDKMobile]) ||
-                !FPaths::FileExists(AirPackBinaryCommand[AppleSDKMac]) ||
-                !FPaths::FileExists(AirPackBinaryCommand[AppleSDKMobile]))
+                !FPaths::FileExists(MetalPackBinaryCommand[AppleSDKMac]) ||
+                !FPaths::FileExists(MetalPackBinaryCommand[AppleSDKMobile]))
 			{
 				UE_LOG(LogMetalCompilerSetup, Warning, TEXT("Missing toolchain binaries in %s."), *ToolchainBase);
 				return EMetalToolchainStatus::ToolchainNotFound;
@@ -975,15 +975,15 @@ FMetalCompilerToolchain::EMetalToolchainStatus FMetalCompilerToolchain::DoWindow
 	MetalLibBinaryCommand[AppleSDKMac] = MacBinaryBasePath / MetalLibraryBinary;
 	MetalLibBinaryCommand[AppleSDKMobile] = IOSBinaryBasePath / MetalLibraryBinary;
 
-	AirPackBinaryCommand[AppleSDKMac] = MacBinaryBasePath / AirPackBinary;
-    AirPackBinaryCommand[AppleSDKMobile] = IOSBinaryBasePath / AirPackBinary;
+	MetalPackBinaryCommand[AppleSDKMac] = MacBinaryBasePath / MetalPackBinary;
+	MetalPackBinaryCommand[AppleSDKMobile] = IOSBinaryBasePath / MetalPackBinary;
     
 	if (!FPaths::FileExists(MetalArBinaryCommand[AppleSDKMac]) ||
 		!FPaths::FileExists(MetalArBinaryCommand[AppleSDKMobile]) ||
 		!FPaths::FileExists(MetalLibBinaryCommand[AppleSDKMac]) ||
 		!FPaths::FileExists(MetalLibBinaryCommand[AppleSDKMobile]) ||
-        !FPaths::FileExists(AirPackBinaryCommand[AppleSDKMac]) ||
-        !FPaths::FileExists(AirPackBinaryCommand[AppleSDKMobile]))
+        !FPaths::FileExists(MetalPackBinaryCommand[AppleSDKMac]) ||
+        !FPaths::FileExists(MetalPackBinaryCommand[AppleSDKMobile]))
 	{
 		if (GCheckCompilerToolChainSetup > 0)
 		{
@@ -1027,18 +1027,18 @@ bool FMetalCompilerToolchain::ExecMetalLib(EAppleSDKType SDK, const TCHAR* Param
 	return ExecGenericCommand(*this->MetalLibBinaryCommand[SDK], Parameters, OutReturnCode, OutStdOut, OutStdErr);
 }
 
-bool FMetalCompilerToolchain::ExecAirPack(EAppleSDKType SDK, const TCHAR* Parameters, int32* OutReturnCode, FString* OutStdOut, FString* OutStdErr) const
+bool FMetalCompilerToolchain::ExecMetalPack(EAppleSDKType SDK, const TCHAR* Parameters, int32* OutReturnCode, FString* OutStdOut, FString* OutStdErr) const
 {
     check(this->bToolchainBinariesPresent);
 #if PLATFORM_MAC
-    if (this->AirPackBinaryCommand[SDK].IsEmpty())
+    if (this->MetalPackBinaryCommand[SDK].IsEmpty())
     {
-        FString BuiltParams = FString::Printf(TEXT("--sdk %s %s %s"), *SDKToString(SDK), *this->AirPackBinary, Parameters);
+        FString BuiltParams = FString::Printf(TEXT("--sdk %s %s %s"), *SDKToString(SDK), *this->MetalPackBinary, Parameters);
         return ExecGenericCommand(*XcrunPath, *BuiltParams, OutReturnCode, OutStdOut, OutStdErr);
     }
     else
 #endif
-    return ExecGenericCommand(*this->AirPackBinaryCommand[SDK], Parameters, OutReturnCode, OutStdOut, OutStdErr);
+    return ExecGenericCommand(*this->MetalPackBinaryCommand[SDK], Parameters, OutReturnCode, OutStdOut, OutStdErr);
 }
 
 bool FMetalCompilerToolchain::ExecMetalAr(EAppleSDKType SDK, const TCHAR* ScriptFile, int32* OutReturnCode, FString* OutStdOut, FString* OutStdErr) const
diff --git a/Engine/Source/Developer/Apple/MetalShaderFormat/Private/MetalShaderFormat.h b/Engine/Source/Developer/Apple/MetalShaderFormat/Private/MetalShaderFormat.h
--- a/Engine/Source/Developer/Apple/MetalShaderFormat/Private/MetalShaderFormat.h
+++ b/Engine/Source/Developer/Apple/MetalShaderFormat/Private/MetalShaderFormat.h
@@ -79,8 +79,8 @@ public:
 	bool ExecMetalLib(EAppleSDKType SDK, const TCHAR* Parameters, int32* OutReturnCode, FString* OutStdOut, FString* OutStdErr) const;
 	// Executes metal-ar for 'SDK' on the local or remote machine, depending on configuration
 	bool ExecMetalAr(EAppleSDKType SDK, const TCHAR* ScriptFile, int32* OutReturnCode, FString* OutStdOut, FString* OutStdErr) const;
-	// Executes air-pack for 'SDK' on the local or remote machine, depending on configuration
-    bool ExecAirPack(EAppleSDKType SDK, const TCHAR* Parameters, int32* OutReturnCode, FString* OutStdOut, FString* OutStdErr) const;
+	// Executes metal-pack for 'SDK' on the local or remote machine, depending on configuration
+    bool ExecMetalPack(EAppleSDKType SDK, const TCHAR* Parameters, int32* OutReturnCode, FString* OutStdOut, FString* OutStdErr) const;
 
 	// This toolchain is set up correctly and ready to use.
 	bool IsCompilerAvailable() const
@@ -195,8 +195,8 @@ public:
 	static FString MetalArBinary;
 	// The name of the metal binary packager - metallib
 	static FString MetalLibraryBinary;
-	// The name of air-pack
-    static FString AirPackBinary;
+	// The name of metal-pack
+    static FString MetalPackBinary;
 
 	// The extension of the mapping from shader to metallib for shared material libraries - .metalmap
 	static FString MetalMapExtension;
@@ -234,8 +234,8 @@ private:
 	// The command string to invoke 'metal-ar'
 	FString MetalArBinaryCommand[AppleSDKCount];
 
-	// The command string to invoke 'air-pack'
-    FString AirPackBinaryCommand[AppleSDKCount];
+	// The command string to invoke 'metal-pack'
+    FString MetalPackBinaryCommand[AppleSDKCount];
 
 	// The compiler version string, parsed out of metal -v
 	FString	MetalCompilerVersionString[AppleSDKCount];
