Metal RHI

Adding Metal-strip to lib processing to reduce metallib sizes (see around 10% reduction).

#DisabledValidations Preflight

diff --git a/Engine/Source/Developer/Apple/MetalShaderFormat/Private/MetalShaderCompiler.cpp b/Engine/Source/Developer/Apple/MetalShaderFormat/Private/MetalShaderCompiler.cpp
--- a/Engine/Source/Developer/Apple/MetalShaderFormat/Private/MetalShaderCompiler.cpp
+++ b/Engine/Source/Developer/Apple/MetalShaderFormat/Private/MetalShaderCompiler.cpp
@@ -1325,7 +1325,7 @@ uint64 AppendShader_Metal(FString const& WorkingDir, const FSHAHash& Hash, TArra
 	return Id;
 }
 
-bool FinalizeLibrary_Metal(FName const& Format, FString const& WorkingDir, FString const& LibraryPath, TSet<uint64> const& Shaders, class FString const& DebugOutputDir)
+bool FinalizeLibrary_Metal(FName const& Format, FString const& WorkingDir, FString const& LibraryPath, TSet<uint64> const& Shaders, class FString const& DebugOutputDir, bool bStripLibs)
 {
 	bool bOK = false;
 
@@ -1397,71 +1397,72 @@ bool FinalizeLibrary_Metal(FName const& Format, FString const& WorkingDir, FStri
 
             {
                 // handle compile error
-                if (ReturnCode == 0 && bArchiveFileValid)
-                {
-                    // MetalPack each batch.
-                    FString BatchMergeArgs;
-                    for (int32 BatchIdx = 0; BatchIdx < MetalPackArgsBatches.Num(); BatchIdx++)
-                    {
-                        FString BatchOutputFile = FString::Printf(TEXT("MetalPackBatch_%d_"), BatchIdx);
-                        FString BatchOutputPath = FPaths::CreateTempFilename(*FullyQualifiedWorkingDir, *BatchOutputFile);
-
-                        BatchMergeArgs += FString::Printf(TEXT("\"%s\" "), *BatchOutputPath);
-
-                        UE_LOG(LogMetalShaderCompiler, Display, TEXT("[%d/%d] %s"), (BatchIdx + 1), MetalPackArgsBatches.Num(), *BatchOutputPath);
-
-                        FString MetalPackParams = FString::Printf(TEXT("-pack-metallibs internal -pack-descriptors internal -pack-reflections internal %s -o \"%s\""), *MetalPackArgsBatches[BatchIdx], *BatchOutputPath);
-                        ReturnCode = 0;
-                        Results.Empty();
-                        Errors.Empty();
-
-                        bool bSuccess = Toolchain->ExecMetalPack(SDK, *MetalPackParams, &ReturnCode, &Results, &Errors);
-
-                        // handle compile error
-                        if (!bSuccess || ReturnCode != 0)
-                        {
-                            UE_LOG(LogShaders, Error, TEXT("Archiving failed: metal-pack failed with code %d: %s %s"), ReturnCode, *Results, *Errors);
-                        }
-                    }
-
-                    // Final pass: merge all batches into a single lib.
-                    UE_LOG(LogMetalShaderCompiler, Display, TEXT("Post-processing archive for shader platform: %s"), *Format.GetPlainNameString());
-
-                    FString LocalMetalLibPath = LibraryPath;
-
-                    if (FPaths::FileExists(LocalMetalLibPath))
-                    {
-                        UE_LOG(LogMetalShaderCompiler, Warning, TEXT("Archiving warning: target metallib already exists and will be overwritten: %s"), *LocalMetalLibPath);
-                    }
-
-                    FString MetalPackParams = FString::Printf(TEXT("-pack-metallibs internal -pack-descriptors internal -pack-reflections internal %s -o \"%s\""), *BatchMergeArgs, *LocalMetalLibPath);
-                    ReturnCode = 0;
-                    Results.Empty();
-                    Errors.Empty();
-
-                    bool bSuccess = Toolchain->ExecMetalPack(SDK, *MetalPackParams, &ReturnCode, &Results, &Errors);
-
-                    // handle compile error
-                    if (bSuccess && ReturnCode == 0)
-                    {
-                        check(LocalMetalLibPath == LibraryPath);
-
-                        bOK = (IFileManager::Get().FileSize(*LibraryPath) > 0);
-
-                        if (!bOK)
-                        {
-                            UE_LOG(LogShaders, Error, TEXT("Archiving failed: failed to copy to local destination: %s"), *LibraryPath);
-                        }
-                    }
-                    else
-                    {
-                        UE_LOG(LogShaders, Error, TEXT("Archiving failed: metal-pack failed with code %d: %s %s"), ReturnCode, *Results, *Errors);
-                    }
-                }
-                else
-                {
-                    UE_LOG(LogShaders, Error, TEXT("Archiving failed: no valid input for metal-pack."));
-                }
+                if (ReturnCode != 0 || !bArchiveFileValid)
+				{
+					UE_LOG(LogShaders, Error, TEXT("Archiving failed: no valid input for metal-pack."));
+					return false;
+				}
+				
+				// MetalPack each batch.
+				FString BatchMergeArgs;
+				for (int32 BatchIdx = 0; BatchIdx < MetalPackArgsBatches.Num(); BatchIdx++)
+				{
+					FString BatchOutputFile = FString::Printf(TEXT("MetalPackBatch_%d_"), BatchIdx);
+					FString BatchOutputPath = FPaths::CreateTempFilename(*FullyQualifiedWorkingDir, *BatchOutputFile);
+
+					BatchMergeArgs += FString::Printf(TEXT("\"%s\" "), *BatchOutputPath);
+
+                    UE_LOG(LogMetalShaderCompiler, Display, TEXT("[%d/%d] %s"), (BatchIdx + 1), MetalPackArgsBatches.Num(), *BatchOutputPath);
+
+					FString MetalPackParams = FString::Printf(TEXT("-pack-metallibs internal -pack-descriptors internal -pack-reflections internal %s -o \"%s\""), *MetalPackArgsBatches[BatchIdx], *BatchOutputPath);
+					ReturnCode = 0;
+					Results.Empty();
+					Errors.Empty();
+
+					bool bSuccess = Toolchain->ExecMetalPack(SDK, *MetalPackParams, &ReturnCode, &Results, &Errors);
+
+					// handle compile error
+					if (!bSuccess || ReturnCode != 0)
+					{
+						UE_LOG(LogShaders, Error, TEXT("Archiving failed: metal-pack failed with code %d: %s %s"), ReturnCode, *Results, *Errors);
+					}
+				}
+
+				// Final pass: merge all batches into a single lib.
+				UE_LOG(LogMetalShaderCompiler, Display, TEXT("Post-processing archive for shader platform: %s"), *Format.GetPlainNameString());
+
+				FString LocalMetalLibPath = LibraryPath;
+
+				if (FPaths::FileExists(LocalMetalLibPath))
+				{
+					UE_LOG(LogMetalShaderCompiler, Warning, TEXT("Archiving warning: target metallib already exists and will be overwritten: %s"), *LocalMetalLibPath);
+				}
+
+				FString MetalPackParams = FString::Printf(TEXT("-pack-metallibs internal -pack-descriptors internal -pack-reflections internal %s -o \"%s\""), *BatchMergeArgs, *LocalMetalLibPath);
+				ReturnCode = 0;
+				Results.Empty();
+				Errors.Empty();
+
+				bool bSuccess = Toolchain->ExecMetalPack(SDK, *MetalPackParams, &ReturnCode, &Results, &Errors);
+
+				// handle compile error
+				if (bSuccess && ReturnCode == 0)
+				{
+					check(LocalMetalLibPath == LibraryPath);
+					
+					bOK = (IFileManager::Get().FileSize(*LibraryPath) > 0);
+					
+					if (!bOK)
+					{
+						UE_LOG(LogShaders, Error, TEXT("Archiving failed: failed to copy to local destination: %s"), *LibraryPath);
+						return bOK;
+					}
+				}
+				else
+				{
+					UE_LOG(LogShaders, Error, TEXT("Archiving failed: metal-pack failed with code %d: %s %s"), ReturnCode, *Results, *Errors);
+					return bOK;
+				}
             }
         }
         else
@@ -1525,45 +1526,72 @@ bool FinalizeLibrary_Metal(FName const& Format, FString const& WorkingDir, FStri
 #endif
 		{
 			// handle compile error
-			if (ReturnCode == 0 && bArchiveFileValid)
+			if (ReturnCode != 0 || !bArchiveFileValid)
 			{
-				UE_LOG(LogMetalShaderCompiler, Display, TEXT("Post-processing archive for shader platform: %s"), *Format.GetPlainNameString());
+				UE_LOG(LogShaders, Error, TEXT("Archiving failed: no valid input for metallib."));
+				return false;
+			}
+			
+			UE_LOG(LogMetalShaderCompiler, Display, TEXT("Post-processing archive for shader platform: %s"), *Format.GetPlainNameString());
 
-				FString LocalMetalLibPath = LibraryPath;
+			FString LocalMetalLibPath = LibraryPath;
 
-				if (FPaths::FileExists(LocalMetalLibPath))
-				{
-					UE_LOG(LogMetalShaderCompiler, Warning, TEXT("Archiving warning: target metallib already exists and will be overwritten: %s"), *LocalMetalLibPath);
-				}
+			if (FPaths::FileExists(LocalMetalLibPath))
+			{
+				UE_LOG(LogMetalShaderCompiler, Warning, TEXT("Archiving warning: target metallib already exists and will be overwritten: %s"), *LocalMetalLibPath);
+			}
+		
+			FString MetallibParams = FString::Printf(TEXT("-o \"%s\" \"%s\""), *LocalMetalLibPath, *LocalArchivePath);
+			ReturnCode = 0;
+			Results.Empty();
+			Errors.Empty();
 			
-				FString MetallibParams = FString::Printf(TEXT("-o \"%s\" \"%s\""), *LocalMetalLibPath, *LocalArchivePath);
-				ReturnCode = 0;
-				Results.Empty();
-				Errors.Empty();
-				
-				bool bSuccess = Toolchain->ExecMetalLib(SDK, *MetallibParams, &ReturnCode, &Results, &Errors);
-	
-				// handle compile error
-				if (bSuccess && ReturnCode == 0)
-				{
-					check(LocalMetalLibPath == LibraryPath);
+			bool bSuccess = Toolchain->ExecMetalLib(SDK, *MetallibParams, &ReturnCode, &Results, &Errors);
 
-					bOK = (IFileManager::Get().FileSize(*LibraryPath) > 0);
+			// handle compile error
+			if (!bSuccess || ReturnCode != 0)
+			{
+				UE_LOG(LogShaders, Error, TEXT("Archiving failed: metallib failed with code %d: %s %s"), ReturnCode, *Results, *Errors);
+				return false;
+			}
+			
+			check(LocalMetalLibPath == LibraryPath);
 
-					if (!bOK)
-					{
-						UE_LOG(LogShaders, Error, TEXT("Archiving failed: failed to copy to local destination: %s"), *LibraryPath);
-					}
+			bOK = (IFileManager::Get().FileSize(*LibraryPath) > 0);
+
+			if (!bOK)
+			{
+				UE_LOG(LogShaders, Error, TEXT("Archiving failed: failed to copy to local destination: %s"), *LibraryPath);
+				return false;
+			}
+			
+			if(bStripLibs)
+			{
+				// Archive build phase - like unix ar, build metal archive from all the object files
+				UE_LOG(LogMetalShaderCompiler, Display, TEXT("Stripping lib %s Size: %d"), *LibraryPath, IFileManager::Get().FileSize(*LibraryPath));
+				
+				FString MetalStripParams = FString::Printf(TEXT("-S -T %s -o \"%s\""), *LibraryPath, *LibraryPath);
+				
+				bSuccess = Toolchain->ExecMetalStrip(SDK, *MetalStripParams, &ReturnCode, &Results, &Errors);
+				
+				if (!bSuccess || ReturnCode != 0)
+				{
+					UE_LOG(LogShaders, Error, TEXT("Stripping failed: metallib failed with code %d: %s %s"), ReturnCode, *Results, *Errors);
+					return false;
 				}
-				else
+				
+				UE_LOG(LogMetalShaderCompiler, Display, TEXT("Stripped size of %s Size: %d"), *LibraryPath, IFileManager::Get().FileSize(*LibraryPath));
+				
+				check(LocalMetalLibPath == LibraryPath);
+				
+				bOK = (IFileManager::Get().FileSize(*LibraryPath) > 0);
+				
+				if (!bOK)
 				{
-					UE_LOG(LogShaders, Error, TEXT("Archiving failed: metallib failed with code %d: %s %s"), ReturnCode, *Results, *Errors);
+					UE_LOG(LogShaders, Error, TEXT("Stripping failed: failed to copy to local destination: %s"), *LibraryPath);
+					return bOK;
 				}
 			}
-			else
-			{
-				UE_LOG(LogShaders, Error, TEXT("Archiving failed: no valid input for metallib."));
-			}
 		}
 	}
 	else
diff --git a/Engine/Source/Developer/Apple/MetalShaderFormat/Private/MetalShaderFormat.cpp b/Engine/Source/Developer/Apple/MetalShaderFormat/Private/MetalShaderFormat.cpp
--- a/Engine/Source/Developer/Apple/MetalShaderFormat/Private/MetalShaderFormat.cpp
+++ b/Engine/Source/Developer/Apple/MetalShaderFormat/Private/MetalShaderFormat.cpp
@@ -42,7 +42,7 @@ extern void OutputMetalDebugData(const FShaderCompilerInput& Input, const FShade
 
 extern bool StripShader_Metal(TArray<uint8>& Code, class FString const& DebugPath, bool const bNative);
 extern uint64 AppendShader_Metal(class FString const& ArchivePath, const FSHAHash& Hash, TArray<uint8>& Code);
-extern bool FinalizeLibrary_Metal(class FName const& Format, class FString const& ArchivePath, class FString const& LibraryPath, TSet<uint64> const& Shaders, class FString const& DebugOutputDir);
+extern bool FinalizeLibrary_Metal(class FName const& Format, class FString const& ArchivePath, class FString const& LibraryPath, TSet<uint64> const& Shaders, class FString const& DebugOutputDir, bool bStripLibs);
 
 /** Version for shader format, this becomes part of the DDC key. */
 static const FGuid UE_SHADER_METAL_VER = FGuid("B0DC25EF-C34D-437A-94E5-5E5146AF1B9A");
@@ -246,7 +246,10 @@ public:
 			// Enqueue the library compilation as a task so we can go wide
 			FGraphEventRef CompletionFence = FFunctionGraphTask::CreateAndDispatchWhenReady([ShaderFormatName, ArchivePath, LibraryPath, PartialShaders, DebugOutputDir, &CompiledLibraries]()
 				{
-					if (FinalizeLibrary_Metal(ShaderFormatName, ArchivePath, LibraryPath, PartialShaders, DebugOutputDir))
+					static const auto CVar = IConsoleManager::Get().FindConsoleVariable(TEXT("r.Shaders.Symbols"));
+					bool bStripLibrary = CVar->GetInt() == 0;
+				
+					if (FinalizeLibrary_Metal(ShaderFormatName, ArchivePath, LibraryPath, PartialShaders, DebugOutputDir, bStripLibrary))
 					{
 						FPlatformAtomics::InterlockedIncrement(&CompiledLibraries);
 					}
@@ -600,11 +603,13 @@ FString FMetalCompilerToolchain::MetalFrontendBinary(TEXT("metal.exe"));
 FString FMetalCompilerToolchain::MetalArBinary(TEXT("metal-ar.exe"));
 FString FMetalCompilerToolchain::MetalLibraryBinary(TEXT("metallib.exe"));
 FString FMetalCompilerToolchain::MetalPackBinary(TEXT("metal-pack.exe"));
+FString FMetalCompilerToolchain::MetalStripBinary(TEXT("metal-strip.exe"));
 #else
 FString FMetalCompilerToolchain::MetalFrontendBinary(TEXT("metal"));
 FString FMetalCompilerToolchain::MetalArBinary(TEXT("metal-ar"));
 FString FMetalCompilerToolchain::MetalLibraryBinary(TEXT("metallib"));
 FString FMetalCompilerToolchain::MetalPackBinary(TEXT("metal-pack"));
+FString FMetalCompilerToolchain::MetalStripBinary(TEXT("metal-strip"));
 #endif
 
 FString FMetalCompilerToolchain::MetalMapExtension(TEXT(".metalmap"));
@@ -867,6 +872,9 @@ FMetalCompilerToolchain::EMetalToolchainStatus FMetalCompilerToolchain::DoMacNat
 
 			MetalPackBinaryCommand[AppleSDKMac] = ToolchainBase / TEXT("macos") / TEXT("bin") / MetalPackBinary;
             MetalPackBinaryCommand[AppleSDKMobile] = ToolchainBase / TEXT("ios") / TEXT("bin") / MetalPackBinary;
+			
+			MetalStripBinaryCommand[AppleSDKMac] = ToolchainBase / TEXT("macos") / TEXT("bin") / MetalStripBinary;
+			MetalStripBinaryCommand[AppleSDKMobile] = ToolchainBase / TEXT("ios") / TEXT("bin") / MetalStripBinary;
 
 			if (!FPaths::FileExists(MetalArBinaryCommand[AppleSDKMac]) ||
 				!FPaths::FileExists(MetalArBinaryCommand[AppleSDKMobile]) ||
@@ -874,7 +882,9 @@ FMetalCompilerToolchain::EMetalToolchainStatus FMetalCompilerToolchain::DoMacNat
 				!FPaths::FileExists(MetalLibBinaryCommand[AppleSDKMobile]) ||
 				!FPaths::FileExists(MetalLibBinaryCommand[AppleSDKMobile]) ||
                 !FPaths::FileExists(MetalPackBinaryCommand[AppleSDKMac]) ||
-                !FPaths::FileExists(MetalPackBinaryCommand[AppleSDKMobile]))
+                !FPaths::FileExists(MetalPackBinaryCommand[AppleSDKMobile]) ||
+				!FPaths::FileExists(MetalStripBinaryCommand[AppleSDKMac]) ||
+				!FPaths::FileExists(MetalStripBinaryCommand[AppleSDKMobile]))
 			{
 				UE_LOG(LogMetalCompilerSetup, Warning, TEXT("Missing toolchain binaries in %s."), *ToolchainBase);
 				return EMetalToolchainStatus::ToolchainNotFound;
@@ -976,14 +986,19 @@ FMetalCompilerToolchain::EMetalToolchainStatus FMetalCompilerToolchain::DoWindow
 	MetalLibBinaryCommand[AppleSDKMobile] = IOSBinaryBasePath / MetalLibraryBinary;
 
 	MetalPackBinaryCommand[AppleSDKMac] = MacBinaryBasePath / MetalPackBinary;
-	MetalPackBinaryCommand[AppleSDKMobile] = IOSBinaryBasePath / MetalPackBinary;
+    MetalPackBinaryCommand[AppleSDKMobile] = IOSBinaryBasePath / MetalPackBinary;
+	
+	MetalStripBinaryCommand[AppleSDKMac] = MacBinaryBasePath / MetalStripBinary;
+	MetalStripBinaryCommand[AppleSDKMobile] = IOSBinaryBasePath / MetalStripBinary;
     
 	if (!FPaths::FileExists(MetalArBinaryCommand[AppleSDKMac]) ||
 		!FPaths::FileExists(MetalArBinaryCommand[AppleSDKMobile]) ||
 		!FPaths::FileExists(MetalLibBinaryCommand[AppleSDKMac]) ||
 		!FPaths::FileExists(MetalLibBinaryCommand[AppleSDKMobile]) ||
         !FPaths::FileExists(MetalPackBinaryCommand[AppleSDKMac]) ||
-        !FPaths::FileExists(MetalPackBinaryCommand[AppleSDKMobile]))
+        !FPaths::FileExists(MetalPackBinaryCommand[AppleSDKMobile]) ||
+		!FPaths::FileExists(MetalStripBinaryCommand[AppleSDKMac]) ||
+		!FPaths::FileExists(MetalStripBinaryCommand[AppleSDKMobile]))
 	{
 		if (GCheckCompilerToolChainSetup > 0)
 		{
@@ -1072,6 +1087,20 @@ bool FMetalCompilerToolchain::ExecMetalAr(EAppleSDKType SDK, const TCHAR* Script
 	return bSuccess;
 }
 
+bool FMetalCompilerToolchain::ExecMetalStrip(EAppleSDKType SDK, const TCHAR* Parameters, int32* OutReturnCode, FString* OutStdOut, FString* OutStdErr) const
+{
+	check(this->bToolchainBinariesPresent);
+#if PLATFORM_MAC
+	if (this->MetalStripBinaryCommand[SDK].IsEmpty())
+	{
+		FString BuiltParams = FString::Printf(TEXT("--sdk %s %s %s"), *SDKToString(SDK), *this->MetalStripBinary, Parameters);
+		return ExecGenericCommand(*XcrunPath, *BuiltParams, OutReturnCode, OutStdOut, OutStdErr);
+	}
+	else
+#endif
+	return ExecGenericCommand(*this->MetalStripBinaryCommand[SDK], Parameters, OutReturnCode, OutStdOut, OutStdErr);
+}
+
 bool FMetalCompilerToolchain::ExecGenericCommand(const TCHAR* Command, const TCHAR* Params, int32* OutReturnCode, FString* OutStdOut, FString* OutStdErr, bool bIsConsoleApp) const
 {
 #if PLATFORM_WINDOWS
diff --git a/Engine/Source/Developer/Apple/MetalShaderFormat/Private/MetalShaderFormat.h b/Engine/Source/Developer/Apple/MetalShaderFormat/Private/MetalShaderFormat.h
--- a/Engine/Source/Developer/Apple/MetalShaderFormat/Private/MetalShaderFormat.h
+++ b/Engine/Source/Developer/Apple/MetalShaderFormat/Private/MetalShaderFormat.h
@@ -81,6 +81,7 @@ public:
 	bool ExecMetalAr(EAppleSDKType SDK, const TCHAR* ScriptFile, int32* OutReturnCode, FString* OutStdOut, FString* OutStdErr) const;
 	// Executes metal-pack for 'SDK' on the local or remote machine, depending on configuration
     bool ExecMetalPack(EAppleSDKType SDK, const TCHAR* Parameters, int32* OutReturnCode, FString* OutStdOut, FString* OutStdErr) const;
+	bool ExecMetalStrip(EAppleSDKType SDK, const TCHAR* Parameters, int32* OutReturnCode, FString* OutStdOut, FString* OutStdErr) const;
 
 	// This toolchain is set up correctly and ready to use.
 	bool IsCompilerAvailable() const
@@ -197,6 +198,8 @@ public:
 	static FString MetalLibraryBinary;
 	// The name of metal-pack
     static FString MetalPackBinary;
+	// The name of metal-strip
+	static FString MetalStripBinary;
 
 	// The extension of the mapping from shader to metallib for shared material libraries - .metalmap
 	static FString MetalMapExtension;
@@ -236,6 +239,8 @@ private:
 
 	// The command string to invoke 'metal-pack'
     FString MetalPackBinaryCommand[AppleSDKCount];
+	// The command string to invoke 'metal-strip'
+	FString MetalStripBinaryCommand[AppleSDKCount];
 
 	// The compiler version string, parsed out of metal -v
 	FString	MetalCompilerVersionString[AppleSDKCount];
