Disable stripping of metallibs on Mac platforms due to crashes on Intel CPU's

diff --git a/Engine/Source/Developer/Apple/MetalShaderFormat/Private/MetalShaderFormat.cpp b/Engine/Source/Developer/Apple/MetalShaderFormat/Private/MetalShaderFormat.cpp
--- a/Engine/Source/Developer/Apple/MetalShaderFormat/Private/MetalShaderFormat.cpp
+++ b/Engine/Source/Developer/Apple/MetalShaderFormat/Private/MetalShaderFormat.cpp
@@ -244,10 +244,10 @@ public:
 			}
 
 			// Enqueue the library compilation as a task so we can go wide
-			FGraphEventRef CompletionFence = FFunctionGraphTask::CreateAndDispatchWhenReady([ShaderFormatName, ArchivePath, LibraryPath, PartialShaders, DebugOutputDir, &CompiledLibraries]()
+			FGraphEventRef CompletionFence = FFunctionGraphTask::CreateAndDispatchWhenReady([ShaderFormatName, ArchivePath, LibraryPath, PartialShaders, DebugOutputDir, &CompiledLibraries, Platform]()
 				{
 					static const auto CVar = IConsoleManager::Get().FindConsoleVariable(TEXT("r.Shaders.Symbols"));
-					bool bStripLibrary = CVar->GetInt() == 0;
+					bool bStripLibrary = FDataDrivenShaderPlatformInfo::GetIsMobile(Platform) && CVar->GetInt() == 0;
 				
 					if (FinalizeLibrary_Metal(ShaderFormatName, ArchivePath, LibraryPath, PartialShaders, DebugOutputDir, bStripLibrary))
 					{
