Fix "Too many positional arguments specified" error in metal-strip

diff --git a/Engine/Source/Developer/Apple/MetalShaderFormat/Private/MetalShaderCompiler.cpp b/Engine/Source/Developer/Apple/MetalShaderFormat/Private/MetalShaderCompiler.cpp
--- a/Engine/Source/Developer/Apple/MetalShaderFormat/Private/MetalShaderCompiler.cpp
+++ b/Engine/Source/Developer/Apple/MetalShaderFormat/Private/MetalShaderCompiler.cpp
@@ -1570,7 +1570,7 @@ bool FinalizeLibrary_Metal(FName const& Format, FString const& WorkingDir, FStri
 				// Archive build phase - like unix ar, build metal archive from all the object files
 				UE_LOG(LogMetalShaderCompiler, Display, TEXT("Stripping lib %s Size: %d"), *LibraryPath, IFileManager::Get().FileSize(*LibraryPath));
 				
-				FString MetalStripParams = FString::Printf(TEXT("-S -T %s -o \"%s\""), *LibraryPath, *LibraryPath);
+                FString MetalStripParams = FString::Printf(TEXT("-S -T -o=\"%s\" \"%s\""), *LibraryPath, *LibraryPath);
 				
 				bSuccess = Toolchain->ExecMetalStrip(SDK, *MetalStripParams, &ReturnCode, &Results, &Errors);
 				
