Delay metal version checking

diff --git a/Engine/Source/Developer/Apple/MetalShaderFormat/Private/MetalShaderFormat.cpp b/Engine/Source/Developer/Apple/MetalShaderFormat/Private/MetalShaderFormat.cpp
--- a/Engine/Source/Developer/Apple/MetalShaderFormat/Private/MetalShaderFormat.cpp
+++ b/Engine/Source/Developer/Apple/MetalShaderFormat/Private/MetalShaderFormat.cpp
@@ -52,7 +52,8 @@ class FMetalShaderFormat : public UE::ShaderCompilerCommon::FBaseShaderFormat
 public:
 	FMetalShaderFormat()
 	{
-		FMetalCompilerToolchain::CreateAndInit();
+		// Delayed until first use
+		//FMetalCompilerToolchain::CreateAndInit();
 	}
 	virtual ~FMetalShaderFormat()
 	{
diff --git a/Engine/Source/Developer/Apple/MetalShaderFormat/Private/MetalShaderFormat.h b/Engine/Source/Developer/Apple/MetalShaderFormat/Private/MetalShaderFormat.h
--- a/Engine/Source/Developer/Apple/MetalShaderFormat/Private/MetalShaderFormat.h
+++ b/Engine/Source/Developer/Apple/MetalShaderFormat/Private/MetalShaderFormat.h
@@ -125,6 +125,10 @@ public:
 	// The Singleton Toolchain object
 	static const FMetalCompilerToolchain* Get()
 	{
+		if (Singleton == nullptr)
+		{
+			CreateAndInit();
+		}
 		return Singleton;
 	}
 
