Delay metal version checking (supplemental)

Fix potential null reference in Teardown.

diff --git a/Engine/Source/Developer/Apple/MetalShaderFormat/Private/MetalShaderFormat.cpp b/Engine/Source/Developer/Apple/MetalShaderFormat/Private/MetalShaderFormat.cpp
--- a/Engine/Source/Developer/Apple/MetalShaderFormat/Private/MetalShaderFormat.cpp
+++ b/Engine/Source/Developer/Apple/MetalShaderFormat/Private/MetalShaderFormat.cpp
@@ -632,9 +632,12 @@ void FMetalCompilerToolchain::CreateAndInit()
 
 void FMetalCompilerToolchain::Destroy()
 {
-	Singleton->Teardown();
-	delete Singleton;
-	Singleton = nullptr;
+	if (Singleton != nullptr)
+	{
+		Singleton->Teardown();
+		delete Singleton;
+		Singleton = nullptr;
+	}
 }
 
 EShaderPlatform FMetalCompilerToolchain::MetalShaderFormatToLegacyShaderPlatform(FName ShaderFormat)
