Metal - Fix number of shader per library with MSC when shader symbols are enabled

#tests Local cook of CitySample

diff --git a/Engine/Source/Developer/Apple/MetalShaderFormat/Private/MetalShaderFormat.cpp b/Engine/Source/Developer/Apple/MetalShaderFormat/Private/MetalShaderFormat.cpp
--- a/Engine/Source/Developer/Apple/MetalShaderFormat/Private/MetalShaderFormat.cpp
+++ b/Engine/Source/Developer/Apple/MetalShaderFormat/Private/MetalShaderFormat.cpp
@@ -185,7 +185,7 @@ public:
 		const EShaderPlatform Platform = FMetalCompilerToolchain::MetalShaderFormatToLegacyShaderPlatform(ShaderFormatName);
 		if(FDataDrivenShaderPlatformInfo::GetSupportsBindless(Platform)
 		   && UE::ShaderCompiler::GetBindlessConfiguration(Platform) != ERHIBindlessConfiguration::Disabled
-		   && ShouldGenerateShaderSymbolsInfo(ShaderFormatName))
+		   && ShouldGenerateShaderSymbols(ShaderFormatName))
 		{
 			NumShadersPerLibrary = 1000;
 		}
