Disable metal version checking for non-metal platforms.

diff --git a/Engine/Source/Developer/Apple/MetalShaderFormat/Private/MetalShaderFormat.cpp b/Engine/Source/Developer/Apple/MetalShaderFormat/Private/MetalShaderFormat.cpp
--- a/Engine/Source/Developer/Apple/MetalShaderFormat/Private/MetalShaderFormat.cpp
+++ b/Engine/Source/Developer/Apple/MetalShaderFormat/Private/MetalShaderFormat.cpp
@@ -52,8 +52,7 @@ class FMetalShaderFormat : public UE::ShaderCompilerCommon::FBaseShaderFormat
 public:
 	FMetalShaderFormat()
 	{
-		// Delayed until first use
-		//FMetalCompilerToolchain::CreateAndInit();
+		FMetalCompilerToolchain::CreateAndInit();
 	}
 	virtual ~FMetalShaderFormat()
 	{
@@ -449,7 +448,6 @@ class FMetalShaderFormatModule : public IShaderFormatModule
 public:
 	virtual ~FMetalShaderFormatModule()
 	{
-		
 		Singleton = nullptr;
 	}
 
@@ -465,8 +463,11 @@ public:
 
 	virtual void ShutdownModule() override
 	{
-		delete Singleton;
-		Singleton = nullptr;
+		if (Singleton != nullptr)
+		{
+			delete Singleton;
+			Singleton = nullptr;
+		}
 	}
 };
 
@@ -929,6 +930,16 @@ FMetalCompilerToolchain::EMetalToolchainStatus FMetalCompilerToolchain::DoMacNat
 #if PLATFORM_WINDOWS
 FMetalCompilerToolchain::EMetalToolchainStatus FMetalCompilerToolchain::DoWindowsSetup()
 {
+	FString TargetPlatform;
+	if (FParse::Value(FCommandLine::Get(), TEXT("TargetPlatform="), TargetPlatform)) {
+		if (!TargetPlatform.StartsWith(TEXT("IOS")) && !TargetPlatform.StartsWith(TEXT("Mac"))) {
+			FPlatformMisc::SetEnvironmentVar(TEXT("NO_METAL"),TEXT("PLEASE"));
+			return EMetalToolchainStatus::ToolchainNotFound;
+		}
+	} else if (FPlatformMisc::GetEnvironmentVariable(TEXT("NO_METAL")) == TEXT("PLEASE")) {
+		return EMetalToolchainStatus::ToolchainNotFound;
+	}
+
 	int32 Result = 0;
 	
 	FString ToolchainBase;
diff --git a/Engine/Source/Developer/Apple/MetalShaderFormat/Private/MetalShaderFormat.h b/Engine/Source/Developer/Apple/MetalShaderFormat/Private/MetalShaderFormat.h
--- a/Engine/Source/Developer/Apple/MetalShaderFormat/Private/MetalShaderFormat.h
+++ b/Engine/Source/Developer/Apple/MetalShaderFormat/Private/MetalShaderFormat.h
@@ -125,10 +125,6 @@ public:
 	// The Singleton Toolchain object
 	static const FMetalCompilerToolchain* Get()
 	{
-		if (Singleton == nullptr)
-		{
-			CreateAndInit();
-		}
 		return Singleton;
 	}
 
