Handle tvOS and case

diff --git a/Engine/Source/Developer/Apple/MetalShaderFormat/Private/MetalShaderFormat.cpp b/Engine/Source/Developer/Apple/MetalShaderFormat/Private/MetalShaderFormat.cpp
--- a/Engine/Source/Developer/Apple/MetalShaderFormat/Private/MetalShaderFormat.cpp
+++ b/Engine/Source/Developer/Apple/MetalShaderFormat/Private/MetalShaderFormat.cpp
@@ -931,12 +931,18 @@ FMetalCompilerToolchain::EMetalToolchainStatus FMetalCompilerToolchain::DoMacNat
 FMetalCompilerToolchain::EMetalToolchainStatus FMetalCompilerToolchain::DoWindowsSetup()
 {
 	FString TargetPlatform;
-	if (FParse::Value(FCommandLine::Get(), TEXT("TargetPlatform="), TargetPlatform)) {
-		if (!TargetPlatform.StartsWith(TEXT("IOS")) && !TargetPlatform.StartsWith(TEXT("Mac"))) {
+	if (FParse::Value(FCommandLine::Get(), TEXT("TargetPlatform="), TargetPlatform))
+	{
+		if (!TargetPlatform.StartsWith(TEXT("IOS"), ESearchCase::IgnoreCase) && 
+		    !TargetPlatform.StartsWith(TEXT("Mac"), ESearchCase::IgnoreCase) &&
+		    !TargetPlatform.StartsWith(TEXT("tvOS"), ESearchCase::IgnoreCase))
+		{
 			FPlatformMisc::SetEnvironmentVar(TEXT("NO_METAL"),TEXT("PLEASE"));
 			return EMetalToolchainStatus::ToolchainNotFound;
 		}
-	} else if (FPlatformMisc::GetEnvironmentVariable(TEXT("NO_METAL")) == TEXT("PLEASE")) {
+	}
+	else if (FPlatformMisc::GetEnvironmentVariable(TEXT("NO_METAL")) == TEXT("PLEASE"))
+	{
 		return EMetalToolchainStatus::ToolchainNotFound;
 	}
 
