MetalShaderFormat - Add two new CVars to set the number of shaders per metallibs in bindful and bindless with symbols

#tests local build + preflight
#Synced-CL Unknown

diff --git a/Engine/Source/Developer/Apple/MetalShaderFormat/Private/MetalShaderFormat.cpp b/Engine/Source/Developer/Apple/MetalShaderFormat/Private/MetalShaderFormat.cpp
--- a/Engine/Source/Developer/Apple/MetalShaderFormat/Private/MetalShaderFormat.cpp
+++ b/Engine/Source/Developer/Apple/MetalShaderFormat/Private/MetalShaderFormat.cpp
@@ -35,6 +35,15 @@ static FAutoConsoleVariableRef CVarCheckCompilerToolChainSetup(
 	ECVF_Default
 );
 
+static TAutoConsoleVariable<int32> CVarNumShadersPerLibrary(
+	TEXT("Metal.NumShadersPerLibrary"), 10000,
+	TEXT("Set the maximum number of shaders per library."),
+	ECVF_Default);
+
+static TAutoConsoleVariable<int32> CVarNumShadersPerLibraryBindlessWithSymbols(
+	TEXT("Metal.NumShadersPerLibraryBindlessWithSymbols"), 1000,
+	TEXT("Set the maximum number of shaders per bindless library when symbols are enabled."),
+	ECVF_Default);
 
 extern bool PreprocessMetalShader(const FShaderCompilerInput& Input, const FShaderCompilerEnvironment& Environment, FShaderPreprocessOutput& PreprocessOutput);
 extern void CompileMetalShader(const FShaderCompilerInput& Input, const FShaderPreprocessOutput& InPreprocessOutput, FShaderCompilerOutput& Output);
@@ -168,7 +177,7 @@ public:
 		const TArray<FSharedBuffer>& ShaderCode,
 		TArray<FString>* OutputFiles) const override final
 	{
-		int32 NumShadersPerLibrary = 10000;
+		int32 NumShadersPerLibrary = CVarNumShadersPerLibrary.GetValueOnAnyThread();
 		
 		check(LibraryName.Len() > 0);
 
@@ -180,13 +189,13 @@ public:
 
 		check(ShaderFormatName == NAME_SF_METAL_ES3_1_IOS || ShaderFormatName == NAME_SF_METAL_SM5_IOS || ShaderFormatName == NAME_SF_METAL_ES3_1_TVOS || ShaderFormatName == NAME_SF_METAL_SM5_TVOS || ShaderFormatName == NAME_SF_METAL_SM5 || ShaderFormatName == NAME_SF_METAL_SM6 || ShaderFormatName == NAME_SF_METAL_SIM || ShaderFormatName == NAME_SF_METAL_ES3_1);
 
-		// Bindless needs a lower limit of shaders per library as the packing process takes a significant amount of RAM
+		// Bindless needs a lower limit of shaders per library as the packing process takes a significant amount of RAM when symbols are enabled
 		const EShaderPlatform Platform = FMetalCompilerToolchain::MetalShaderFormatToLegacyShaderPlatform(ShaderFormatName);
 		if(FDataDrivenShaderPlatformInfo::GetSupportsBindless(Platform)
 		   && UE::ShaderCompiler::GetBindlessConfiguration(Platform) != ERHIBindlessConfiguration::Disabled
 		   && ShouldGenerateShaderSymbols(ShaderFormatName))
 		{
-			NumShadersPerLibrary = 1000;
+			NumShadersPerLibrary = CVarNumShadersPerLibraryBindlessWithSymbols.GetValueOnAnyThread();
 		}
 		
 		const FString ArchivePath = (WorkingDirectory / ShaderFormatAndShaderPlatformName.GetPlainNameString());
