Shader formats that depend on ShaderConductor need to include the ShaderConductor version hash when calculating their version hashes.

#Synced-CL 49623085

diff --git a/Engine/Source/Developer/Apple/MetalShaderFormat/Private/MetalShaderFormat.cpp b/Engine/Source/Developer/Apple/MetalShaderFormat/Private/MetalShaderFormat.cpp
--- a/Engine/Source/Developer/Apple/MetalShaderFormat/Private/MetalShaderFormat.cpp
+++ b/Engine/Source/Developer/Apple/MetalShaderFormat/Private/MetalShaderFormat.cpp
@@ -1,6 +1,7 @@
 // Copyright Epic Games, Inc. All Rights Reserved.
 
 #include "MetalShaderFormat.h"
+#include "DXCWrapper.h"
 #include "Modules/ModuleInterface.h"
 #include "Modules/ModuleManager.h"
 #include "Interfaces/IShaderFormat.h"
@@ -58,8 +59,11 @@ static const FGuid UE_SHADER_METAL_VER = FGuid("B0DC25EF-C34D-437A-94E5-5E5146AF
 
 class FMetalShaderFormat : public UE::ShaderCompilerCommon::FBaseShaderFormat 
 {
+	const uint32 ShaderConductorVersionHash;
+
 public:
-	FMetalShaderFormat()
+	FMetalShaderFormat(uint32 InShaderConductorVersionHash)
+		: ShaderConductorVersionHash(InShaderConductorVersionHash)
 	{
 		FMetalCompilerToolchain::CreateAndInit();
 	}
@@ -106,6 +110,7 @@ public:
 
 		Result = HashCombine(Result, GetTypeHash(HLSLCC_VersionMinor));
 		Result = HashCombine(Result, GetTypeHash(UE_SHADER_METAL_VER));
+		Result = HashCombine(Result, GetTypeHash(ShaderConductorVersionHash));
 
 		return Result;
 	}
@@ -452,7 +457,7 @@ public:
 
 static IShaderFormat* Singleton = nullptr;
 
-class FMetalShaderFormatModule : public IShaderFormatModule
+class FMetalShaderFormatModule : public IShaderFormatModule, public FShaderConductorModuleWrapper
 {
 public:
 	virtual ~FMetalShaderFormatModule()
@@ -467,7 +472,7 @@ public:
 
 	virtual void StartupModule() override
 	{
-		Singleton = new FMetalShaderFormat();
+		Singleton = new FMetalShaderFormat(FShaderConductorModuleWrapper::GetModuleVersionHash());
 	}
 
 	virtual void ShutdownModule() override
diff --git a/Engine/Source/Developer/ShaderFormatOpenGL/Private/ShaderFormatOpenGL.cpp b/Engine/Source/Developer/ShaderFormatOpenGL/Private/ShaderFormatOpenGL.cpp
--- a/Engine/Source/Developer/ShaderFormatOpenGL/Private/ShaderFormatOpenGL.cpp
+++ b/Engine/Source/Developer/ShaderFormatOpenGL/Private/ShaderFormatOpenGL.cpp
@@ -3,6 +3,7 @@
 
 #include "ShaderFormatOpenGL.h"
 
+#include "DXCWrapper.h"
 #include "ShaderCompilerCommon.h"
 #include "ShaderPreprocessor.h"
 #include "ShaderPreprocessTypes.h"
@@ -35,7 +36,14 @@ class FShaderFormatGLSL : public UE::ShaderCompilerCommon::FBaseShaderFormat
 		check(Format == NAME_GLSL_150_ES3_1 || Format == NAME_GLSL_ES3_1_ANDROID);
 	}
 
+	const uint32 ShaderConductorVersionHash;
+
 public:
+	FShaderFormatGLSL(uint32 InShaderConductorVersionHash)
+		: ShaderConductorVersionHash(InShaderConductorVersionHash)
+	{
+	}
+
 	virtual uint32 GetVersion(FName Format) const override
 	{
 		CheckFormat(Format);
@@ -49,6 +57,7 @@ public:
 
 		uint32 Version = GetTypeHash(HLSLCC_VersionMinor);
 		Version = HashCombine(Version, BaseHash);
+		Version = HashCombine(Version, GetTypeHash(ShaderConductorVersionHash));
 
 		return Version;
 	}
@@ -125,7 +134,7 @@ public:
 
 static IShaderFormat* Singleton = nullptr;
 
-class FShaderFormatOpenGLModule : public IShaderFormatModule
+class FShaderFormatOpenGLModule : public IShaderFormatModule, public FShaderConductorModuleWrapper
 {
 public:
 	virtual ~FShaderFormatOpenGLModule() override
@@ -137,7 +146,7 @@ public:
 	{
 		if (!Singleton)
 		{
-			Singleton = new FShaderFormatGLSL();
+			Singleton = new FShaderFormatGLSL(FShaderConductorModuleWrapper::GetModuleVersionHash());
 		}
 		return Singleton;
 	}
