HandleReflectedGlobalConstantBufferMember needs a ShaderParameterParser argument for future work.

#Synced-CL 49734186

diff --git a/Engine/Source/Developer/Apple/MetalShaderFormat/Private/MetalCompileShaderMSC.cpp b/Engine/Source/Developer/Apple/MetalShaderFormat/Private/MetalCompileShaderMSC.cpp
--- a/Engine/Source/Developer/Apple/MetalShaderFormat/Private/MetalCompileShaderMSC.cpp
+++ b/Engine/Source/Developer/Apple/MetalShaderFormat/Private/MetalCompileShaderMSC.cpp
@@ -29,6 +29,7 @@ THIRD_PARTY_INCLUDES_END
 extern void BuildMetalShaderOutput(
 	FShaderCompilerOutput& ShaderOutput,
 	const FShaderCompilerInput& ShaderInput,
+	const FShaderParameterParser& ShaderParameterParser,
 	const ANSICHAR* InShaderSource,
 	uint32 SourceLen,
 	uint32 SourceCRCLen,
@@ -159,7 +160,7 @@ static IRDescriptorRange1 CreateDescriptorRange(const uint32 NumDescriptors)
 	return DescRange;
 }
 
-static void ProcessReflection(ID3D12ShaderReflection* ShaderReflection, const uint32 BoundResources, const FShaderCompilerInput& Input, FShaderCompilerOutput& Output, CrossCompiler::FHlslccHeaderWriter& CCHeaderWriter, FMetalShaderOutputMetaData& OutputData, uint32& NumCBVs)
+static void ProcessReflection(ID3D12ShaderReflection* ShaderReflection, const uint32 BoundResources, const FShaderCompilerInput& Input, const FShaderParameterParser& ShaderParameterParser, FShaderCompilerOutput& Output, CrossCompiler::FHlslccHeaderWriter& CCHeaderWriter, FMetalShaderOutputMetaData& OutputData, uint32& NumCBVs)
 {
 	NumCBVs = 0;
 
@@ -257,6 +258,8 @@ static void ProcessReflection(ID3D12ShaderReflection* ShaderReflection, const ui
 						InsertRange(CBRanges, BindIndex, SourceOffset, DestSize, DestCBIndex, DestCBPrecision, DestOffset);
 						
 						HandleReflectedGlobalConstantBufferMember(
+							Input,
+							ShaderParameterParser,
 							VariableName,
 							BindIndex,
 							VariableDesc.StartOffset,
@@ -405,7 +408,7 @@ static IMalloc* GetDxcMalloc()
 }
 #endif // PLATFORM_WINDOWS
 
-static bool ReflectDXILAndBuildDescriptorRanges(const TArray<uint32>& DXILReflection, const FShaderCompilerInput& Input, FShaderCompilerOutput& Output, CrossCompiler::FHlslccHeaderWriter& CCHeaderWriter, FMetalShaderOutputMetaData& OutputData, uint32& NumCBVs, uint32& NumInstructions)
+static bool ReflectDXILAndBuildDescriptorRanges(const TArray<uint32>& DXILReflection, const FShaderCompilerInput& Input, const FShaderParameterParser& ShaderParameterParser, FShaderCompilerOutput& Output, CrossCompiler::FHlslccHeaderWriter& CCHeaderWriter, FMetalShaderOutputMetaData& OutputData, uint32& NumCBVs, uint32& NumInstructions)
 {
 	// Reflect DXIL
 	TRefCountPtr<IDxcUtils> Utils;
@@ -448,7 +451,7 @@ static bool ReflectDXILAndBuildDescriptorRanges(const TArray<uint32>& DXILReflec
 		constexpr uint32 DxbcToDxilInstructionRatio[2] = { 119, 307 };
 		NumInstructions = NumInstructions * DxbcToDxilInstructionRatio[0] / DxbcToDxilInstructionRatio[1];
 
-		ProcessReflection(ShaderReflection.GetReference(), ShaderDesc.BoundResources, Input, Output, CCHeaderWriter, OutputData, NumCBVs);
+		ProcessReflection(ShaderReflection.GetReference(), ShaderDesc.BoundResources, Input, ShaderParameterParser, Output, CCHeaderWriter, OutputData, NumCBVs);
 
 		// Vertex Input
 		for (uint32 InputIndex = 0; InputIndex < ShaderDesc.InputParameters; InputIndex++)
@@ -537,6 +540,7 @@ struct FMetalShaderParameterParserPlatformConfiguration : public FShaderParamete
 
 void FMetalCompileShaderMSC::DoCompileMetalShader(
 	const FShaderCompilerInput& Input,
+	const FShaderParameterParser& InShaderParameterParser, // todo: why isn't this one used?
 	FShaderCompilerOutput& Output,
 	const FString& InPreprocessedShader,
 	uint32 VersionEnum,
@@ -588,6 +592,14 @@ void FMetalCompileShaderMSC::DoCompileMetalShader(
 	FMetalShaderBytecode MetalBytecode;
 
 #if PLATFORM_MAC || PLATFORM_WINDOWS
+	FMetalShaderParameterParserPlatformConfiguration PlatformConfiguration;
+	FShaderParameterParser ShaderParameterParser(PlatformConfiguration);
+	if (!ShaderParameterParser.ParseAndModify(Input, Output.Errors, PreprocessedShader))
+	{
+		// The FShaderParameterParser will add any relevant errors.
+		return;
+	}
+
 	{
 		std::string EntryPointNameAnsi(TCHAR_TO_UTF8(*Input.EntryPointName));
 
@@ -599,14 +611,6 @@ void FMetalCompileShaderMSC::DoCompileMetalShader(
 		Options.TargetEnvironment = CrossCompiler::FShaderConductorOptions::ETargetEnvironment::Vulkan_1_2;
 
 		Options.bEnable16bitTypes = true;
-
-		FMetalShaderParameterParserPlatformConfiguration PlatformConfiguration;
-		FShaderParameterParser ShaderParameterParser(PlatformConfiguration);
-		if (!ShaderParameterParser.ParseAndModify(Input, Output.Errors, PreprocessedShader))
-		{
-			// The FShaderParameterParser will add any relevant errors.
-			return;
-		}
 		
 		TArray<FString> ExtraArgs;
 		
@@ -702,7 +706,7 @@ void FMetalCompileShaderMSC::DoCompileMetalShader(
 		FCStringAnsi::Snprintf(MainCRC, 25, "Main_%0.8x_%0.8x", CRCLen, CRC);
 		
 		// Build shader metadata and root signature parameters
-		bool bSuccessfulReflection = ReflectDXILAndBuildDescriptorRanges(DxilData, Input, Output, CCHeaderWriter, OutputData, NumCBVs, Output.NumInstructions);
+		bool bSuccessfulReflection = ReflectDXILAndBuildDescriptorRanges(DxilData, Input, ShaderParameterParser, Output, CCHeaderWriter, OutputData, NumCBVs, Output.NumInstructions);
 		check(bSuccessfulReflection);
 		
 		// Build root parameters
@@ -926,7 +930,7 @@ void FMetalCompileShaderMSC::DoCompileMetalShader(
 	}
 
 	Output.Target = Input.Target;
-	BuildMetalShaderOutput(Output, Input, MetalSource.c_str(), MetalSource.length(), CRCLen, CRC, VersionEnum, *Standard, *MinOSVersion, Output.Errors, OutputData.TypedBuffers, OutputData.InvariantBuffers, OutputData.TypedUAVs, OutputData.ConstantBuffers, bAllowFastIntrinsics
+	BuildMetalShaderOutput(Output, Input, ShaderParameterParser, MetalSource.c_str(), MetalSource.length(), CRCLen, CRC, VersionEnum, *Standard, *MinOSVersion, Output.Errors, OutputData.TypedBuffers, OutputData.InvariantBuffers, OutputData.TypedUAVs, OutputData.ConstantBuffers, bAllowFastIntrinsics
 		, NumCBVs, OutputSizeVS, MaxInputPrimitivesPerMeshThreadgroupGS, bUsesDiscard, ReflectionJSON, MetalBytecode
 	);
 }
diff --git a/Engine/Source/Developer/Apple/MetalShaderFormat/Private/MetalCompileShaderMSC.h b/Engine/Source/Developer/Apple/MetalShaderFormat/Private/MetalCompileShaderMSC.h
--- a/Engine/Source/Developer/Apple/MetalShaderFormat/Private/MetalCompileShaderMSC.h
+++ b/Engine/Source/Developer/Apple/MetalShaderFormat/Private/MetalCompileShaderMSC.h
@@ -14,16 +14,18 @@ class FMetalCompileShaderMSC
 {
 public:
 	static void DoCompileMetalShader(
-									 const FShaderCompilerInput& Input,
-									 FShaderCompilerOutput& Output,
-									 const FString& InPreprocessedShader,
-									 uint32 VersionEnum,
-									 EMetalGPUSemantics Semantics,
-									 uint32 MaxUnrollLoops,
-									 EShaderFrequency Frequency,
-									 bool bDumpDebugInfo,
-									 const FString& Standard,
-									 const FString& MinOSVersion);
+		const FShaderCompilerInput& Input,
+		const FShaderParameterParser& ShaderParameterParser,
+		FShaderCompilerOutput& Output,
+		const FString& InPreprocessedShader,
+		uint32 VersionEnum,
+		EMetalGPUSemantics Semantics,
+		uint32 MaxUnrollLoops,
+		EShaderFrequency Frequency,
+		bool bDumpDebugInfo,
+		const FString& Standard,
+		const FString& MinOSVersion
+	);
 };
 
 #endif
diff --git a/Engine/Source/Developer/Apple/MetalShaderFormat/Private/MetalCompileShaderSPIRV.cpp b/Engine/Source/Developer/Apple/MetalShaderFormat/Private/MetalCompileShaderSPIRV.cpp
--- a/Engine/Source/Developer/Apple/MetalShaderFormat/Private/MetalCompileShaderSPIRV.cpp
+++ b/Engine/Source/Developer/Apple/MetalShaderFormat/Private/MetalCompileShaderSPIRV.cpp
@@ -18,6 +18,7 @@
 extern void BuildMetalShaderOutput(
 	FShaderCompilerOutput& ShaderOutput,
 	const FShaderCompilerInput& ShaderInput,
+	const FShaderParameterParser& ShaderParameterParser,
 	const ANSICHAR* InShaderSource,
 	uint32 SourceLen,
 	uint32 SourceCRCLen,
@@ -174,6 +175,7 @@ static void Patch16bitInHlslSource(const FShaderCompilerInput& Input, FAnsiStrin
 
 void FMetalCompileShaderSPIRV::DoCompileMetalShader(
 	const FShaderCompilerInput& Input,
+	const FShaderParameterParser& ShaderParameterParser,
 	FShaderCompilerOutput& Output,
 	const FString& InPreprocessedShader,
 	uint32 VersionEnum,
@@ -1362,7 +1364,7 @@ void FMetalCompileShaderSPIRV::DoCompileMetalShader(
 	if (Result != 0)
 	{
 		Output.Target = Input.Target;
-		BuildMetalShaderOutput(Output, Input, MetalSource.c_str(), MetalSource.length(), CRCLen, CRC, VersionEnum, *Standard, *MinOSVersion, Output.Errors, OutputData.TypedBuffers, OutputData.InvariantBuffers, OutputData.TypedUAVs, OutputData.ConstantBuffers, bAllowFastIntrinsics
+		BuildMetalShaderOutput(Output, Input, ShaderParameterParser, MetalSource.c_str(), MetalSource.length(), CRCLen, CRC, VersionEnum, *Standard, *MinOSVersion, Output.Errors, OutputData.TypedBuffers, OutputData.InvariantBuffers, OutputData.TypedUAVs, OutputData.ConstantBuffers, bAllowFastIntrinsics
 #if UE_METAL_USE_METAL_SHADER_CONVERTER
 		  , 0, 0, 0, false, nullptr, FMetalShaderBytecode()
 #endif
diff --git a/Engine/Source/Developer/Apple/MetalShaderFormat/Private/MetalCompileShaderSPIRV.h b/Engine/Source/Developer/Apple/MetalShaderFormat/Private/MetalCompileShaderSPIRV.h
--- a/Engine/Source/Developer/Apple/MetalShaderFormat/Private/MetalCompileShaderSPIRV.h
+++ b/Engine/Source/Developer/Apple/MetalShaderFormat/Private/MetalCompileShaderSPIRV.h
@@ -11,14 +11,16 @@ class FMetalCompileShaderSPIRV
 {
 public:
 	static void DoCompileMetalShader(
-									 const FShaderCompilerInput& Input,
-									 FShaderCompilerOutput& Output,
-									 const FString& InPreprocessedShader,
-									 uint32 VersionEnum,
-									 EMetalGPUSemantics Semantics,
-									 uint32 MaxUnrollLoops,
-									 EShaderFrequency Frequency,
-									 bool bDumpDebugInfo,
-									 const FString& Standard,
-									 const FString& MinOSVersion);
+		const FShaderCompilerInput& Input,
+		const FShaderParameterParser& ShaderParameterParser,
+		FShaderCompilerOutput& Output,
+		const FString& InPreprocessedShader,
+		uint32 VersionEnum,
+		EMetalGPUSemantics Semantics,
+		uint32 MaxUnrollLoops,
+		EShaderFrequency Frequency,
+		bool bDumpDebugInfo,
+		const FString& Standard,
+		const FString& MinOSVersion
+	);
 };
diff --git a/Engine/Source/Developer/Apple/MetalShaderFormat/Private/MetalShaderCompiler.cpp b/Engine/Source/Developer/Apple/MetalShaderFormat/Private/MetalShaderCompiler.cpp
--- a/Engine/Source/Developer/Apple/MetalShaderFormat/Private/MetalShaderCompiler.cpp
+++ b/Engine/Source/Developer/Apple/MetalShaderFormat/Private/MetalShaderCompiler.cpp
@@ -238,6 +238,7 @@ static const int32 Str##PrefixLen = FCStringAnsi::Strlen(Str##Prefix)
 void BuildMetalShaderOutput(
 	FShaderCompilerOutput& ShaderOutput,
 	const FShaderCompilerInput& ShaderInput,
+	const FShaderParameterParser& ShaderParameterParser,
 	const ANSICHAR* InShaderSource,
 	uint32 SourceLen,
 	uint32 SourceCRCLen,
@@ -405,6 +406,8 @@ void BuildMetalShaderOutput(
 	for (const CrossCompiler::FHlslccHeader::FPackedGlobal& PackedGlobal : CCHeader.PackedGlobals)
 	{
 		HandleReflectedGlobalConstantBufferMember(
+			ShaderInput,
+			ShaderParameterParser,
 			PackedGlobal.Name,
 			PackedGlobal.PackedType,
 			PackedGlobal.Offset * BytesPerComponent,
@@ -435,6 +438,8 @@ void BuildMetalShaderOutput(
 			if (!bUseMetalShaderConverter)
 			{
 				HandleReflectedGlobalConstantBufferMember(
+					ShaderInput,
+					ShaderParameterParser,
 					Member.Name,
 					ConstantBufferIndex,
 					Member.Offset * BytesPerComponent,
@@ -1099,12 +1104,12 @@ void CompileMetalShader(const FShaderCompilerInput& Input, const FShaderPreproce
 #if UE_METAL_USE_METAL_SHADER_CONVERTER
 	if (Input.IsBindlessEnabled())
 	{
-		FMetalCompileShaderMSC::DoCompileMetalShader(Input, Output, PreprocessedSource, VersionEnum, Semantics, MaxUnrollLoops, (EShaderFrequency)Input.Target.Frequency, bDumpDebugInfo, Standard, MinOSVersion);
+		FMetalCompileShaderMSC::DoCompileMetalShader(Input, ShaderParameterParser, Output, PreprocessedSource, VersionEnum, Semantics, MaxUnrollLoops, (EShaderFrequency)Input.Target.Frequency, bDumpDebugInfo, Standard, MinOSVersion);
 	}
 	else
 #endif
 	{
-		FMetalCompileShaderSPIRV::DoCompileMetalShader(Input, Output, PreprocessedSource, VersionEnum, Semantics, MaxUnrollLoops, (EShaderFrequency)Input.Target.Frequency, bDumpDebugInfo, Standard, MinOSVersion);
+		FMetalCompileShaderSPIRV::DoCompileMetalShader(Input, ShaderParameterParser, Output, PreprocessedSource, VersionEnum, Semantics, MaxUnrollLoops, (EShaderFrequency)Input.Target.Frequency, bDumpDebugInfo, Standard, MinOSVersion);
 	}
 	ShaderParameterParser.ValidateShaderParameterTypes(Input, bIsMobile, Output);
 }
diff --git a/Engine/Source/Developer/ShaderCompilerCommon/Private/ShaderCompilerCommon.cpp b/Engine/Source/Developer/ShaderCompilerCommon/Private/ShaderCompilerCommon.cpp
--- a/Engine/Source/Developer/ShaderCompilerCommon/Private/ShaderCompilerCommon.cpp
+++ b/Engine/Source/Developer/ShaderCompilerCommon/Private/ShaderCompilerCommon.cpp
@@ -578,6 +578,28 @@ bool UE::ShaderCompilerCommon::RemoveDeadCode(FShaderSource& InOutPreprocessedSh
 	return UE::ShaderCompilerCommon::RemoveDeadCode(InOutPreprocessedShaderSource, RequiredSymbols, OutErrors);
 }
 
+void HandleReflectedGlobalConstantBufferMember(
+	const FShaderCompilerInput& Input,
+	const FShaderParameterParser& ShaderParameterParser,
+	FStringView InMemberName,
+	uint32 ConstantBufferIndex,
+	int32 ReflectionOffset,
+	int32 ReflectionSize,
+	FShaderCompilerOutput& Output
+)
+{
+	FStringView MemberName = InMemberName;
+	const EShaderParameterType ParameterType = FShaderParameterParser::ParseAndRemoveBindlessParameterPrefix(MemberName);
+
+	Output.ParameterMap.AddParameterAllocation(
+		MemberName,
+		ConstantBufferIndex,
+		ReflectionOffset,
+		ReflectionSize,
+		ParameterType
+	);
+}
+
 void HandleReflectedGlobalConstantBufferMember(
 	FStringView InMemberName,
 	uint32 ConstantBufferIndex,
@@ -656,6 +678,7 @@ void HandleReflectedRootConstantBufferMember(
 	FStringView MemberName = InMemberName;
 	const EShaderParameterType ParameterType = FShaderParameterParser::ParseAndRemoveBindlessParameterPrefix(MemberName);
 
+	// Root constant buffer members only need reflection when they are specific to resources
 	if (ParameterType != EShaderParameterType::LooseData)
 	{
 		Output.ParameterMap.AddParameterAllocation(
diff --git a/Engine/Source/Developer/ShaderCompilerCommon/Public/ShaderCompilerCommon.h b/Engine/Source/Developer/ShaderCompilerCommon/Public/ShaderCompilerCommon.h
--- a/Engine/Source/Developer/ShaderCompilerCommon/Public/ShaderCompilerCommon.h
+++ b/Engine/Source/Developer/ShaderCompilerCommon/Public/ShaderCompilerCommon.h
@@ -208,6 +208,17 @@ SHADERCOMPILERCOMMON_API EUniformBufferMemberReflectionReason ShouldReflectUnifo
 	FStringView UniformBufferName
 );
 
+extern SHADERCOMPILERCOMMON_API void HandleReflectedGlobalConstantBufferMember(
+	const FShaderCompilerInput& Input,
+	const FShaderParameterParser& ShaderParameterParser,
+	FStringView MemberName,
+	uint32 ConstantBufferIndex,
+	int32 ReflectionOffset,
+	int32 ReflectionSize,
+	FShaderCompilerOutput& CompilerOutput
+);
+
+UE_DEPRECATED(5.8, "HandleReflectedGlobalConstantBufferMember now needs Input and ShaderParameterParser parameters for future features")
 extern SHADERCOMPILERCOMMON_API void HandleReflectedGlobalConstantBufferMember(
 	FStringView MemberName,
 	uint32 ConstantBufferIndex,
diff --git a/Engine/Source/Developer/ShaderFormatOpenGL/Private/OpenGLShaderCompiler.cpp b/Engine/Source/Developer/ShaderFormatOpenGL/Private/OpenGLShaderCompiler.cpp
--- a/Engine/Source/Developer/ShaderFormatOpenGL/Private/OpenGLShaderCompiler.cpp
+++ b/Engine/Source/Developer/ShaderFormatOpenGL/Private/OpenGLShaderCompiler.cpp
@@ -571,6 +571,7 @@ static void CompileOfflineInternal(const FShaderCompilerInput& Input, FShaderCom
 void BuildShaderOutputInternal(
 	FShaderCompilerOutput& ShaderOutput,
 	const FShaderCompilerInput& ShaderInput,
+	const FShaderParameterParser& ShaderParameterParser,
 	const ANSICHAR* InShaderSource,
 	int32 SourceLen,
 	GLSLVersion Version)
@@ -704,6 +705,8 @@ void BuildShaderOutputInternal(
 			continue;
 
 		HandleReflectedGlobalConstantBufferMember(
+			ShaderInput,
+			ShaderParameterParser,
 			PackedGlobal.Name,
 			PackedGlobal.PackedType,
 			PackedGlobal.Offset* BytesPerComponent,
@@ -930,7 +933,7 @@ static void ConvertOpenGLVersionFromGLSLVersionInternal(GLSLVersion InVersion, i
  * @param Version - The GLSL language version to target.
  * @param Frequency - The shader stage
  */
-static void PrecompileShaderInternal(FShaderCompilerOutput& ShaderOutput, const FShaderCompilerInput& ShaderInput, const ANSICHAR* ShaderSource, GLSLVersion Version, EHlslShaderFrequency Frequency)
+static void PrecompileShaderInternal(FShaderCompilerOutput& ShaderOutput, const FShaderCompilerInput& ShaderInput, const FShaderParameterParser& ShaderParameterParser, const ANSICHAR* ShaderSource, GLSLVersion Version, EHlslShaderFrequency Frequency)
 {
 	check(ShaderInput.Target.Frequency < SF_NumFrequencies);
 
@@ -971,6 +974,7 @@ static void PrecompileShaderInternal(FShaderCompilerOutput& ShaderOutput, const
 			BuildShaderOutputInternal(
 				ShaderOutput,
 				ShaderInput,
+				ShaderParameterParser,
 				ShaderSource,
 				(int32)SourceLen,
 				Version
@@ -3109,11 +3113,11 @@ void CompileOpenGLShader(const FShaderCompilerInput& Input, const FShaderPreproc
 		}
 
 #if VALIDATE_GLSL_WITH_DRIVER
-		PrecompileShaderInternal(Output, Input, GlslShaderSource, Version, HlslFrequency);
+		PrecompileShaderInternal(Output, Input, ShaderParameterParser, GlslShaderSource, Version, HlslFrequency);
 #else // VALIDATE_GLSL_WITH_DRIVER
 		int32 SourceLen = FCStringAnsi::Strlen(GlslShaderSource); //-V595
 		Output.Target = Input.Target;
-		BuildShaderOutputInternal(Output, Input, GlslShaderSource, SourceLen, Version);
+		BuildShaderOutputInternal(Output, Input, ShaderParameterParser, GlslShaderSource, SourceLen, Version);
 #endif // VALIDATE_GLSL_WITH_DRIVER
 
 		if (bDumpDebugInfo && GlslShaderSource != nullptr)
diff --git a/Engine/Source/Developer/VulkanShaderFormat/Private/SpirVShaderCompiler.inl b/Engine/Source/Developer/VulkanShaderFormat/Private/SpirVShaderCompiler.inl
--- a/Engine/Source/Developer/VulkanShaderFormat/Private/SpirVShaderCompiler.inl
+++ b/Engine/Source/Developer/VulkanShaderFormat/Private/SpirVShaderCompiler.inl
@@ -824,6 +824,8 @@ static bool BuildShaderOutputFromSpirv(
 								{
 									check(BindingTypeCount == 0); // Global constants should always be the first UB
 									HandleReflectedGlobalConstantBufferMember(
+										InternalState.Input,
+										*InternalState.ParameterParser,
 										MemberName,
 										BindingTypeCount,
 										Member.absolute_offset,
diff --git a/Engine/Source/Developer/VulkanShaderFormat/Private/VulkanShaderCompiler.cpp b/Engine/Source/Developer/VulkanShaderFormat/Private/VulkanShaderCompiler.cpp
--- a/Engine/Source/Developer/VulkanShaderFormat/Private/VulkanShaderCompiler.cpp
+++ b/Engine/Source/Developer/VulkanShaderFormat/Private/VulkanShaderCompiler.cpp
@@ -174,7 +174,7 @@ static uint32 GetSizeForType(FStringView TypeName, FStringView ArraySize)
 
 
 // :todo-jn: TEMPORARY EXPERIMENT - will eventually move into preprocessing step
-static uint32 ConvertGlobalsToShaderRecord(const FShaderParameterParser& ShaderParameterParser, const TMap<FStringView, FStringView>& ReplacedGlobals, FString& PreprocessedShaderSource, FShaderCompilerOutput& Output)
+static uint32 ConvertGlobalsToShaderRecord(const FShaderCompilerInput& Input, const FShaderParameterParser& ShaderParameterParser, const TMap<FStringView, FStringView>& ReplacedGlobals, FString& PreprocessedShaderSource, FShaderCompilerOutput& Output)
 {
 	uint32 ShaderRecordGlobalsSize = 0;
 	uint32 ShaderRecordParamCount = 0;
@@ -189,6 +189,8 @@ static uint32 ConvertGlobalsToShaderRecord(const FShaderParameterParser& ShaderP
 		const uint32 ParamSize = GetSizeForType(Info.ParsedType, Info.ParsedArraySize);
 
 		HandleReflectedGlobalConstantBufferMember(
+			Input,
+			ShaderParameterParser,
 			ParamName,
 			0 /*ConstantBufferIndex*/,
 			ShaderRecordGlobalsSize /*ReflectionOffset*/,
@@ -314,7 +316,7 @@ void CompileVulkanShader(const FShaderCompilerInput& Input, const FShaderPreproc
 
 	if (InternalState.bUseBindlessUniformBuffer)
 	{
-		InternalState.ShaderRecordGlobalsSize = ConvertGlobalsToShaderRecord(ShaderParameterParser, ReplacedGlobals, PreprocessedSource, Output);
+		InternalState.ShaderRecordGlobalsSize = ConvertGlobalsToShaderRecord(Input, ShaderParameterParser, ReplacedGlobals, PreprocessedSource, Output);
 		InternalState.AllBindlessUBs = SpirvShaderCompiler::ConvertUBToBindless(PreprocessedSource, SpirvShaderCompiler::EBindlessUniformBufferType::ShaderRecord);
 	}
 
diff --git a/Engine/Source/Developer/Windows/ShaderFormatD3D/Private/D3DShaderCompiler.inl b/Engine/Source/Developer/Windows/ShaderFormatD3D/Private/D3DShaderCompiler.inl
--- a/Engine/Source/Developer/Windows/ShaderFormatD3D/Private/D3DShaderCompiler.inl
+++ b/Engine/Source/Developer/Windows/ShaderFormatD3D/Private/D3DShaderCompiler.inl
@@ -155,6 +155,8 @@ template <typename ID3D1xShaderReflection, typename D3D1x_SHADER_DESC, typename
 							const FString VariableName(VariableDesc.Name);
 
 							HandleReflectedGlobalConstantBufferMember(
+								Input,
+								ShaderParameterParser,
 								VariableName,
 								CBIndex,
 								VariableDesc.StartOffset,
