Metal RHI

- Added StaticShaderBindingLayoutSupport for Metal Raytracing
- Fixed crash in runtime with metal libs

diff --git a/Engine/Config/Mac/DataDrivenPlatformInfo.ini b/Engine/Config/Mac/DataDrivenPlatformInfo.ini
--- a/Engine/Config/Mac/DataDrivenPlatformInfo.ini
+++ b/Engine/Config/Mac/DataDrivenPlatformInfo.ini
@@ -96,6 +96,7 @@ bSupportsRayTracingShaders=true
 bSupportsInlineRayTracing=true
 bSupportsRayTracingProceduralPrimitive=true
 bSupportsRayTracingIndirectInstanceData=true
+StaticShaderBindingLayoutSupport=RayTracingOnly
 bSupportsPathTracing=true
 bSupportsRayTracingCallableShaders=true
 bSupportsHairStrandGeometry=true
diff --git a/Engine/Shaders/Public/Platform/Metal/MetalCommon.ush b/Engine/Shaders/Public/Platform/Metal/MetalCommon.ush
--- a/Engine/Shaders/Public/Platform/Metal/MetalCommon.ush
+++ b/Engine/Shaders/Public/Platform/Metal/MetalCommon.ush
@@ -7,7 +7,7 @@
 #pragma once
 
 // Update this GUID to invalidate shader recompilation for only Metal shaders
-#pragma message("UESHADERMETADATA_VERSION 5B012837-A128-4C05-9F41-0D372A240AA8")
+#pragma message("UESHADERMETADATA_VERSION D1B69FDA-FF60-4414-8915-01A07FC67FF4")
 
 #if VERTEXSHADER && HULLSHADER
 #ifndef TESSELLATIONSHADER
diff --git a/Engine/Source/Developer/Apple/MetalShaderFormat/Private/MetalCompileShaderMSC.cpp b/Engine/Source/Developer/Apple/MetalShaderFormat/Private/MetalCompileShaderMSC.cpp
--- a/Engine/Source/Developer/Apple/MetalShaderFormat/Private/MetalCompileShaderMSC.cpp
+++ b/Engine/Source/Developer/Apple/MetalShaderFormat/Private/MetalCompileShaderMSC.cpp
@@ -761,7 +761,7 @@ void CompileToMetalLib(FMetalShaderBytecode& InMetalBytecode,
 #endif
 								);
 	
-	if (bNeedsAirReflection)
+	if (bNeedsAirReflection || bDumpDebugInfo)
 	{
 		IRShaderReflection* AirReflection = IRShaderReflectionCreate();
 		IRObjectGetReflection(AirBytecode, ShaderStageOut, AirReflection);
@@ -769,9 +769,11 @@ void CompileToMetalLib(FMetalShaderBytecode& InMetalBytecode,
 		if(bDumpDebugInfo)
 		{
 			ReflectionJSON = IRShaderReflectionCopyJSONString(AirReflection);
-			FString ReflectionString = ANSI_TO_TCHAR(ReflectionJSON);
-			DumpDebugShaderText(Input, ReflectionString, TEXT("reflection.json"));
-			checkSlow(ReflectionJSON);
+			if(ReflectionJSON)
+			{
+				FString ReflectionString = ANSI_TO_TCHAR(ReflectionJSON);
+				DumpDebugShaderText(Input, ReflectionString, TEXT("reflection.json"));
+			}
 		}
 		
 		switch (ShaderStageOut)
@@ -952,6 +954,11 @@ void FMetalCompileShaderMSC::DoCompileMetalShader(
 			ExtraArgs.Add(TEXT("-Qembed_debug"));
 			ExtraArgs.Add(TEXT("--ignore-line-directives"));
 		}
+		else if (Input.IsRayTracingShader()) // This is required for RT shaders pending a discussion with Apple
+		{
+			ExtraArgs.Add(TEXT("-Zi"));
+			ExtraArgs.Add(TEXT("-Qembed_debug"));
+		}
 		
 		if (Input.Environment.CompilerFlags.Contains(CFLAG_Debug) || Input.Environment.CompilerFlags.Contains(CFLAG_SkipOptimizationsDXC))
 		{
@@ -1102,14 +1109,34 @@ void FMetalCompileShaderMSC::DoCompileMetalShader(
 			ShaderInfo.UniqueEntryPointName = MainCRC;
 		}
 		
+		// If we are using Shader Binding Layouts, then add to the root signature when linking
+		TArray<IRRootParameter1> LinkingRootParams;
+		
+		if (Input.IsRayTracingShader())
+		{
+			for(uint32_t RegisterIdx = 0; RegisterIdx < FRHIShaderBindingLayout::MaxUniformBufferEntries; RegisterIdx++)
+			{
+				IRRootParameter1 RootParam;
+				RootParam.ParameterType = IRRootParameterTypeCBV;
+				RootParam.ShaderVisibility = IRShaderVisibility::IRShaderVisibilityAll;
+				RootParam.Descriptor.ShaderRegister = RegisterIdx;
+				RootParam.Descriptor.RegisterSpace = UE_HLSL_SPACE_STATIC_SHADER_BINDINGS;
+				RootParam.Descriptor.Flags = IRRootDescriptorFlagDataStaticWhileSetAtExecute;
+				
+				LinkingRootParams.Add(RootParam);
+			}
+		}
+		
+		LinkingRootParams.Append(ShaderInfo.GlobalRootParams);
+		
 		// Create the root signature for air generation.
 		IRVersionedRootSignatureDescriptor RootSignatureDesc;
 		RootSignatureDesc.version = IRRootSignatureVersion_1_1;
 		RootSignatureDesc.desc_1_1.Flags = IRRootSignatureFlagNone;
 		RootSignatureDesc.desc_1_1.pStaticSamplers = StaticSamplerDescs;
 		RootSignatureDesc.desc_1_1.NumStaticSamplers = UE_ARRAY_COUNT(StaticSamplerDescs);
-		RootSignatureDesc.desc_1_1.pParameters = ShaderInfo.GlobalRootParams.GetData();
-		RootSignatureDesc.desc_1_1.NumParameters = ShaderInfo.GlobalRootParams.Num();
+		RootSignatureDesc.desc_1_1.pParameters = LinkingRootParams.GetData();
+		RootSignatureDesc.desc_1_1.NumParameters = LinkingRootParams.Num();
 		
 		IRError* RootSignatureCreationError = nullptr;
 		IRRootSignature* RootSignature = IRRootSignatureCreateFromDescriptor(&RootSignatureDesc, &RootSignatureCreationError);
diff --git a/Engine/Source/Developer/Apple/MetalShaderFormat/Private/MetalShaderCompiler.cpp b/Engine/Source/Developer/Apple/MetalShaderFormat/Private/MetalShaderCompiler.cpp
--- a/Engine/Source/Developer/Apple/MetalShaderFormat/Private/MetalShaderCompiler.cpp
+++ b/Engine/Source/Developer/Apple/MetalShaderFormat/Private/MetalShaderCompiler.cpp
@@ -1539,6 +1539,22 @@ bool FinalizeLibrary_Metal(FName const& Format, FString const& WorkingDir, FStri
 					{
 						UE_LOG(LogShaders, Error, TEXT("Archiving failed: metal-pack failed with code %d: %s %s"), ReturnCode, *Results, *Errors);
 					}
+					
+					if (bStripLibs)
+					{
+						UE_LOG(LogMetalShaderCompiler, Display, TEXT("Stripping lib %s Size: %d"), *BatchOutputPath, IFileManager::Get().FileSize(*BatchOutputPath));
+						
+						FString MetalStripParams = FString::Printf(TEXT("-S -T -o=\"%s\" \"%s\""), *BatchOutputPath, *BatchOutputPath);
+						bSuccess = Toolchain->ExecMetalStrip(SDK, *MetalStripParams, &ReturnCode, &Results, &Errors);
+						
+						// handle compile error
+						if (!bSuccess)
+						{
+							UE_LOG(LogShaders, Error, TEXT("Archiving failed: metal-pack failed with code %d: %s %s"), ReturnCode, *Results, *Errors);
+						}
+						
+						UE_LOG(LogMetalShaderCompiler, Display, TEXT("Stripped size of %s Size: %d"), *BatchOutputPath, IFileManager::Get().FileSize(*BatchOutputPath));
+					}
 				}
 
 				// Final pass: merge all batches into a single lib.
diff --git a/Engine/Source/Developer/Apple/MetalShaderFormat/Private/MetalShaderFormat.cpp b/Engine/Source/Developer/Apple/MetalShaderFormat/Private/MetalShaderFormat.cpp
--- a/Engine/Source/Developer/Apple/MetalShaderFormat/Private/MetalShaderFormat.cpp
+++ b/Engine/Source/Developer/Apple/MetalShaderFormat/Private/MetalShaderFormat.cpp
@@ -320,7 +320,9 @@ public:
 					static const auto CVarSymbols = IConsoleManager::Get().FindConsoleVariable(TEXT("r.Shaders.Symbols"));
 					static const auto CVarExtraData = IConsoleManager::Get().FindConsoleVariable(TEXT("r.Shaders.ExtraData"));
 				
-					bool bStripLibrary = FDataDrivenShaderPlatformInfo::GetIsMobile(Platform) && CVarSymbols->GetInt() == 0 && CVarExtraData->GetInt() == 0;
+					bool bStripLibrary = (FDataDrivenShaderPlatformInfo::GetIsMobile(Platform) || FDataDrivenShaderPlatformInfo::GetSupportsBindless(Platform)) &&
+											CVarSymbols->GetInt() == 0 &&
+											CVarExtraData->GetInt() == 0;
 				
 					if (FinalizeLibrary_Metal(ShaderFormatName, ArchivePath, LibraryPath, PartialShaders, DebugOutputDir, bStripLibrary))
 					{
diff --git a/Engine/Source/Runtime/Apple/MetalRHI/Private/MetalCommands.cpp b/Engine/Source/Runtime/Apple/MetalRHI/Private/MetalCommands.cpp
--- a/Engine/Source/Runtime/Apple/MetalRHI/Private/MetalCommands.cpp
+++ b/Engine/Source/Runtime/Apple/MetalRHI/Private/MetalCommands.cpp
@@ -463,17 +463,35 @@ void FMetalRHICommandContext::RHISetGraphicsPipelineState(FRHIGraphicsPipelineSt
 
 void FMetalRHICommandContext::RHISetStaticUniformBuffers(const FUniformBufferStaticBindings& InUniformBuffers)
 {
-	FMemory::Memzero(GlobalUniformBuffers.GetData(), GlobalUniformBuffers.Num() * sizeof(FRHIUniformBuffer*));
+	FMemory::Memzero(StaticUniformBuffers.GetData(), StaticUniformBuffers.Num() * sizeof(FRHIUniformBuffer*));
 
-	for (int32 Index = 0; Index < InUniformBuffers.GetUniformBufferCount(); ++Index)
+	if (const FRHIShaderBindingLayout* Layout = InUniformBuffers.GetShaderBindingLayout())
 	{
-		GlobalUniformBuffers[InUniformBuffers.GetSlot(Index)] = InUniformBuffers.GetUniformBuffer(Index);
+		check(InUniformBuffers.GetUniformBufferCount() == Layout->GetNumUniformBufferEntries());
+
+		for (int32 Index = 0; Index < InUniformBuffers.GetUniformBufferCount(); ++Index)
+		{
+			StaticUniformBuffers[Index] = InUniformBuffers.GetUniformBuffer(Index);
+			checkf(StaticUniformBuffers[Index], TEXT("Static uniform buffer at index %d is referenced in the shader binding layout but is not provided"), Index);
+		}
+
+		ShaderBindingLayout = Layout;
+	}
+	else
+	{
+		for (int32 Index = 0; Index < InUniformBuffers.GetUniformBufferCount(); ++Index)
+		{
+			FUniformBufferStaticSlot Slot = InUniformBuffers.GetSlot(Index);
+			StaticUniformBuffers[Slot] = InUniformBuffers.GetUniformBuffer(Index);
+		}
+
+		ShaderBindingLayout = nullptr;
 	}
 }
 
 void FMetalRHICommandContext::RHISetStaticUniformBuffer(FUniformBufferStaticSlot InSlot, FRHIUniformBuffer* InBuffer)
 {
-	GlobalUniformBuffers[InSlot] = InBuffer;
+	StaticUniformBuffers[InSlot] = InBuffer;
 }
 
 struct FMetalShaderBinder
diff --git a/Engine/Source/Runtime/Apple/MetalRHI/Private/MetalRHIContext.cpp b/Engine/Source/Runtime/Apple/MetalRHI/Private/MetalRHIContext.cpp
--- a/Engine/Source/Runtime/Apple/MetalRHI/Private/MetalRHIContext.cpp
+++ b/Engine/Source/Runtime/Apple/MetalRHI/Private/MetalRHIContext.cpp
@@ -64,7 +64,7 @@ FMetalRHICommandContext::FMetalRHICommandContext(FMetalDevice& MetalDevice)
 	, RenderPassDesc(nullptr)
 	, bWithinRenderPass(false)
 {
-	GlobalUniformBuffers.AddZeroed(FUniformBufferStaticSlotRegistry::Get().GetSlotCount());
+	StaticUniformBuffers.AddZeroed(FUniformBufferStaticSlotRegistry::Get().GetSlotCount());
 }
 
 FMetalRHICommandContext::~FMetalRHICommandContext()
@@ -90,6 +90,9 @@ void FMetalRHICommandContext::ResetContext()
 	// make sure first SetRenderTarget goes through
 	StateCache.InvalidateRenderTargets();	
 	
+	ShaderBindingLayout = nullptr;
+	FMemory::Memzero(StaticUniformBuffers.GetData(), StaticUniformBuffers.Num() * sizeof(FRHIUniformBuffer*));
+	
 	bIsParallelContext = false;
 }
 
diff --git a/Engine/Source/Runtime/Apple/MetalRHI/Private/MetalRayTracing.cpp b/Engine/Source/Runtime/Apple/MetalRHI/Private/MetalRayTracing.cpp
--- a/Engine/Source/Runtime/Apple/MetalRHI/Private/MetalRayTracing.cpp
+++ b/Engine/Source/Runtime/Apple/MetalRHI/Private/MetalRayTracing.cpp
@@ -2722,7 +2722,7 @@ static void DispatchRays(FMetalRHICommandContext& Context,
 			FMetalUniformBuffer* UB = ResourceCast(GlobalResourceBindings.UniformBuffers[Index]);
 			FMetalBufferPtr Buffer = UB->BackingBuffer;
 			
-			GRSData[Index] = Buffer->GetGPUAddress();	
+			GRSData[Index+RayGenShader->GetGlobalRootSignatureStartIdx()] = Buffer->GetGPUAddress();	
 			ProcessUniformBuffer(GlobalResourceBindings.UniformBuffers[Index], StateCache);
 		}
 	}
@@ -2801,28 +2801,25 @@ static void DispatchRays(FMetalRHICommandContext& Context,
 		}
 	}
 	
-	const TArray<FRHIUniformBuffer*>& GlobalUniformBuffers = Context.GetStaticUniformBuffers();
+	const FRHIShaderBindingLayout& ShaderBindingLayout = Context.GetShaderBindingLayout();
+	const TArray<FRHIUniformBuffer*>& StaticUniformBuffers = Context.GetStaticUniformBuffers();
 	
-	// Bind static uniform buffers to the RGS global root signature.
-	// (this is a copy/paste from ApplyStaticUniformBuffers)
-	for (int32 BufferIndex = 0; BufferIndex < RayGenShader->StaticSlots.Num(); ++BufferIndex)
+	for (uint32 Index = 0; Index < ShaderBindingLayout.GetNumUniformBufferEntries(); ++Index)
 	{
-		const FUniformBufferStaticSlot Slot = RayGenShader->StaticSlots[BufferIndex];
+		const FRHIUniformBufferShaderBindingLayout& LayoutEntry = ShaderBindingLayout.GetUniformBufferEntry(Index);
+		const uint32 RootParameterSlotIndex = LayoutEntry.CBVResourceIndex;
 
-		if (IsUniformBufferStaticSlotValid(Slot) && GRSData[BufferIndex] == 0x0)
-		{
-			FRHIUniformBuffer* Buffer = GlobalUniformBuffers[Slot];
-			if (Buffer)
-			{
-				FMetalUniformBuffer *UB = ResourceCast(Buffer);
-				FMetalBufferPtr BackingBuffer = UB->BackingBuffer;
-				
-				StateCache.CacheOrSkipResourceResidencyUpdate(BackingBuffer->GetMTLBuffer(), EMetalShaderStages::Compute, true);
-				ProcessUniformBuffer(Buffer, StateCache);
-				
-				GRSData[BufferIndex] = BackingBuffer->GetGPUAddress();
-			}
-		}
+		FRHIUniformBuffer* UniformBuffer = StaticUniformBuffers[Index];
+		checkf(UniformBuffer, TEXT("Static uniform buffer at index %d is referenced in the shader binding layout but not provided in the last RHISetStaticUniformBuffers() command"), Index);
+
+		FMetalUniformBuffer *UB = ResourceCast(UniformBuffer);
+		FMetalBufferPtr BackingBuffer = UB->BackingBuffer;
+		
+		StateCache.CacheOrSkipResourceResidencyUpdate(BackingBuffer->GetMTLBuffer(), EMetalShaderStages::Compute, true);
+		ProcessUniformBuffer(UniformBuffer, StateCache);
+		
+		check(GRSEntryCount >= RootParameterSlotIndex + 1);
+		GRSData[RootParameterSlotIndex] = BackingBuffer->GetGPUAddress();
 	}
 
 	// Static samplers always bound to the last slot
diff --git a/Engine/Source/Runtime/Apple/MetalRHI/Private/Shaders/Types/MetalRayShader.cpp b/Engine/Source/Runtime/Apple/MetalRHI/Private/Shaders/Types/MetalRayShader.cpp
--- a/Engine/Source/Runtime/Apple/MetalRHI/Private/Shaders/Types/MetalRayShader.cpp
+++ b/Engine/Source/Runtime/Apple/MetalRHI/Private/Shaders/Types/MetalRayShader.cpp
@@ -137,6 +137,9 @@ FMetalRayShader::FMetalRayShader(FMetalDevice& Device, TArrayView<const uint8> I
 			TEXT("Ray Tracing Shader has %d bits set, which is not the expected count for shader frequency %d"), FMath::CountBits(RayTracingPayloadType), int(InFrequency)
 	);
 
+	
+	GlobalRootSignatureStartIdx = FRHIShaderBindingLayout::MaxUniformBufferEntries;
+	
 	GlobalRootSignature = CreateRootSignature(GlobalRootParams, false);
 	GlobalRootSignatureSize = CalculateRootSignatureSize(GlobalRootSignature);
 	LocalRootSignature = CreateRootSignature(LocalRootParams, true);
@@ -155,8 +158,27 @@ FMetalRayShader::~FMetalRayShader()
 	}
 }
 
-IRRootSignature *FMetalRayShader::CreateRootSignature(const TArray<IRRootParameter1> &RootParams, bool bIsLocal)
+IRRootSignature *FMetalRayShader::CreateRootSignature(const TArray<IRRootParameter1> &InRootParams, bool bIsLocal)
 {
+	TArray<IRRootParameter1> RootParams;
+	
+	if(!bIsLocal)
+	{
+		// Adding static shader bindings
+		for(uint32_t RegisterIdx = 0; RegisterIdx < FRHIShaderBindingLayout::MaxUniformBufferEntries; RegisterIdx++)
+		{
+			IRRootParameter1 RootParam;
+			RootParam.ParameterType = IRRootParameterTypeCBV;
+			RootParam.ShaderVisibility = IRShaderVisibility::IRShaderVisibilityAll;
+			RootParam.Descriptor.ShaderRegister = RegisterIdx;
+			RootParam.Descriptor.RegisterSpace = UE_HLSL_SPACE_STATIC_SHADER_BINDINGS;
+			RootParam.Descriptor.Flags = IRRootDescriptorFlagDataStaticWhileSetAtExecute;
+			
+			RootParams.Add(RootParam);
+		}
+	}
+	RootParams.Append(InRootParams);
+	
 	IRVersionedRootSignatureDescriptor RootSignatureDescriptor;
 	// Create the global root signature for air generation.
 	RootSignatureDescriptor.version = IRRootSignatureVersion_1_1;
diff --git a/Engine/Source/Runtime/Apple/MetalRHI/Private/Shaders/Types/MetalRayShader.h b/Engine/Source/Runtime/Apple/MetalRHI/Private/Shaders/Types/MetalRayShader.h
--- a/Engine/Source/Runtime/Apple/MetalRHI/Private/Shaders/Types/MetalRayShader.h
+++ b/Engine/Source/Runtime/Apple/MetalRHI/Private/Shaders/Types/MetalRayShader.h
@@ -98,6 +98,8 @@ public:
 	inline bool HasVisibleFunctionEntryPoint() const { return RayGen || Miss || Callable || ClosestHit; }
 	inline bool HasIntersectionFunctionEntryPoint() const { return AnyHitAndIntersection; }
 
+	inline uint32 GetGlobalRootSignatureStartIdx() const { return GlobalRootSignatureStartIdx; }
+	
 private:
 	static IRRootSignature *CreateRootSignature(const TArray<IRRootParameter1> &RootParams, bool bIsLocal);
 	static uint32 CalculateRootSignatureSize(IRRootSignature* RootSignature);
@@ -110,6 +112,8 @@ private:
 	TArray<IRRootParameter1> GlobalRootParams;
 	TArray<IRRootParameter1> LocalRootParams;
 
+	uint32 GlobalRootSignatureStartIdx = 0;
+	
 	EShaderFrequency Frequency;
 
 	using FMetalRayGenShader = FMetalBaseRayShader<FRHIRayGenShader, SF_RayGen>;
diff --git a/Engine/Source/Runtime/Apple/MetalRHI/Public/MetalRHIContext.h b/Engine/Source/Runtime/Apple/MetalRHI/Public/MetalRHIContext.h
--- a/Engine/Source/Runtime/Apple/MetalRHI/Public/MetalRHIContext.h
+++ b/Engine/Source/Runtime/Apple/MetalRHI/Public/MetalRHIContext.h
@@ -287,7 +287,7 @@ public:
     
     inline const TArray<FRHIUniformBuffer*>& GetStaticUniformBuffers() const
     {
-        return GlobalUniformBuffers;
+        return StaticUniformBuffers;
     }
 	
 	inline TSharedRef<FMetalQueryBufferPool, ESPMode::ThreadSafe> GetQueryBufferPool()
@@ -315,6 +315,12 @@ public:
 		return bWithinRenderPass;
 	}
 	
+	const FRHIShaderBindingLayout& GetShaderBindingLayout() const
+	{
+		static const FRHIShaderBindingLayout Default;
+		return ShaderBindingLayout ? *ShaderBindingLayout : Default;
+	}
+	
 	void StartCommandBuffer();
 	void EndCommandBuffer();
 	
@@ -345,7 +351,7 @@ protected:
 	bool bIsParallelContext = false;
 	void ResolveTexture(UE::RHICore::FResolveTextureInfo Info);
 
-	TArray<FRHIUniformBuffer*> GlobalUniformBuffers;
+	TArray<FRHIUniformBuffer*> StaticUniformBuffers;
 	
 	// The array of recorded payloads the submission thread will process.
 	// These are returned when the context is finalized.
@@ -356,6 +362,8 @@ protected:
 	
 	FMetalParallelRenderPassInfo* ParallelRenderPassInfo;
 	
+	const FRHIShaderBindingLayout* ShaderBindingLayout = nullptr;
+	
 #if PLATFORM_SUPPORTS_BINDLESS_RENDERING
 	FMetalPendingDescriptorUpdates PendingDescriptorUpdates;
 #endif
diff --git a/Engine/Source/Runtime/Apple/MetalRHI/Public/MetalShaderResources.h b/Engine/Source/Runtime/Apple/MetalRHI/Public/MetalShaderResources.h
--- a/Engine/Source/Runtime/Apple/MetalRHI/Public/MetalShaderResources.h
+++ b/Engine/Source/Runtime/Apple/MetalRHI/Public/MetalShaderResources.h
@@ -401,7 +401,7 @@ struct FMetalBindlessShaderInfo
 	TArray<IRRootParameter1> UnsortedLocalRootParams;
 	TArray<IRRootParameter1> LocalRootParams;
 	
-	static constexpr uint32 NumShaderSpaces = (UE_HLSL_SPACE_RAY_TRACING_SYSTEM + 1);
+	static constexpr uint32 NumShaderSpaces = (UE_HLSL_SPACE_STATIC_SHADER_BINDINGS + 1);
 	
 	// Temporary cache to remove dupes when processing a shader library. Key is the register index; Value is the binding name (for sanity check at compile time).
 	TMap<uint32, FString> RootParamsCache[NumShaderSpaces];
