MetalShaderFormat - Allow the use of Non-native Metal Shader libraries
- Implement support for bSharedMaterialNativeLibraries=False in order to allow the use of FIoStoreShaderCodeArchive on Apple platforms
- Add support for using metal-strip for individual metallibs from SPIRV & MSC path with StripIndividualMetalLibs, disabled by default on Mac because of an issue with Intel Mac
- Add new option AllowsRuntimeShaderCompiling to support runtime shader compiling instead of relying on bSharedMaterialNativeLibraries=False
- Use r.Shaders.Symbols to include source code in metallibs with SPIRV
- Cleanup dead code

#tests Preflight & Local tests

#Synced-CL 50972727

diff --git a/Engine/Config/BaseEngine.ini b/Engine/Config/BaseEngine.ini
--- a/Engine/Config/BaseEngine.ini
+++ b/Engine/Config/BaseEngine.ini
@@ -3218,6 +3218,7 @@ UseFastIntrinsics=False
 ForceFloats=True
 EnableMathOptimisations=True
 MetalLanguageVersion=0
+StripIndividualMetalLibs=True
 bDisableMotionData=False
 bUsesNonExemptEncryption=false
 ITSEncryptionExportComplianceCode=""
@@ -3386,6 +3387,7 @@ MetalLanguageVersion=7
 +TargetedRHIs=SF_METAL_SM6
 UseFastIntrinsics=False
 EnableMathOptimisations=True
+StripIndividualMetalLibs=False
 TargetArchitecture=MacTargetArchitectureUniversal
 EditorTargetArchitecture=MacTargetArchitectureUniversal
 EditorDefaultArchitecture=MacTargetArchitectureHost
@@ -3928,4 +3930,4 @@ bUseSimplifiedUserModel=True
 ; Don't precache global shaders, slate and only precache particles at proxy creation time on ES31 Android and iOS
 +Settings=(ShaderPlatform="OPENGL_ES3_1_ANDROID", GlobalShaderMode=0, bSlateMaterials=False, ParticlePrecachingTime=1)
 +Settings=(ShaderPlatform="VULKAN_ES3_1_ANDROID", GlobalShaderMode=0, bSlateMaterials=False, ParticlePrecachingTime=1)
-+Settings=(ShaderPlatform="METAL_ES3_1_IOS", GlobalShaderMode=0, bSlateMaterials=False, ParticlePrecachingTime=1)
\ No newline at end of file
++Settings=(ShaderPlatform="METAL_ES3_1_IOS", GlobalShaderMode=0, bSlateMaterials=False, ParticlePrecachingTime=1)
diff --git a/Engine/Source/Developer/Apple/MetalShaderFormat/Private/MetalCompileShaderMSC.cpp b/Engine/Source/Developer/Apple/MetalShaderFormat/Private/MetalCompileShaderMSC.cpp
--- a/Engine/Source/Developer/Apple/MetalShaderFormat/Private/MetalCompileShaderMSC.cpp
+++ b/Engine/Source/Developer/Apple/MetalShaderFormat/Private/MetalCompileShaderMSC.cpp
@@ -8,6 +8,8 @@
 #include "Misc/Compression.h"
 #include "Misc/OutputDeviceRedirector.h"
 #include "MetalBackend.h"
+#include "Misc/FileHelper.h"
+#include "Misc/Paths.h"
 #include "RHIDefinitions.h"
 #include "Serialization/MemoryReader.h"
 #include "Serialization/MemoryWriter.h"
@@ -727,6 +729,8 @@ struct FMetalShaderParameterParserPlatformConfiguration : public FShaderParamete
 	const FString HitGroupSystemVertexBufferName;
 };
 
+extern bool ShouldStripIndividualMetalLib(const FShaderCompilerInput& InputCompilerEnvironment);
+
 void CompileToMetalLib(FMetalShaderBytecode& InMetalBytecode,
 					   IRShaderStage& ShaderStageOut,
 					   const FShaderCompilerInput& Input,
@@ -744,10 +748,17 @@ void CompileToMetalLib(FMetalShaderBytecode& InMetalBytecode,
 	
 	IRError* CompileError = nullptr;
 	IRObject* AirBytecode = CompileCallback(&CompileError);
+	ON_SCOPE_EXIT
+	{
+		if (AirBytecode)
+		{
+			IRObjectDestroy(AirBytecode);
+		}
+	};
 	
 	if (!AirBytecode || CompileError != nullptr)
 	{
-		Output.AddErrorf(TEXT("Error: MetalShaderConverter failed to produce air bytecode for '%s' (%s)!"), *Input.EntryPointName, ANSI_TO_TCHAR((const char*)IRErrorGetPayload(CompileError)));
+		Output.AddErrorf(TEXT("Error: MetalShaderConverter failed to produce air bytecode for '%s' (%s)!"), *Input.EntryPointName, (CompileError ? ANSI_TO_TCHAR((const char*)IRErrorGetPayload(CompileError)) : TEXT("CompileError NULL")));
 		Output.bSucceeded = false;
 		
 		return;
@@ -842,6 +853,10 @@ void CompileToMetalLib(FMetalShaderBytecode& InMetalBytecode,
 	
 	// Retrieve the generated .metallib
 	IRMetalLibBinary* GeneratedMetalLib = IRMetalLibBinaryCreate();
+	ON_SCOPE_EXIT
+	{
+		IRMetalLibBinaryDestroy(GeneratedMetalLib);
+	};
 	if (!IRObjectGetMetalLibBinary(AirBytecode, ShaderStageOut, GeneratedMetalLib))
 	{
 		Output.AddErrorf(TEXT("Error: MetalShaderConverter failed to produce a metallib for '%s'!"), *Input.EntryPointName);
@@ -851,21 +866,56 @@ void CompileToMetalLib(FMetalShaderBytecode& InMetalBytecode,
 	}
 	
 	size_t MetalLibSize = IRMetalLibGetBytecodeSize(GeneratedMetalLib);
-	
-	InMetalBytecode.OutputFile.Reserve(MetalLibSize);
 	InMetalBytecode.OutputFile.SetNum(MetalLibSize);
-	
 	size_t OutMetalLibSize = IRMetalLibGetBytecode(GeneratedMetalLib, reinterpret_cast<uint8_t*>(InMetalBytecode.OutputFile.GetData()));
 	
-	InMetalBytecode.ObjectFile.SetNum(MetalLibSize);
-	
-	// Copy the AIR (needed for serialization below)
-	InMetalBytecode.ObjectFile = InMetalBytecode.OutputFile;
+	const bool bUseNativeShaderLibrary = Input.Environment.CompilerFlags.Contains(CFLAG_Archive);
+	if (bUseNativeShaderLibrary)
+	{
+		// Copy the AIR (needed for native shader library)
+		InMetalBytecode.ObjectFile = InMetalBytecode.OutputFile;
+	}
+	else if (ShouldStripIndividualMetalLib(Input))
+	{
+		// For non-native shader library we need to strip the metallib now
+		// First serialize the metallib to disk and then strip it
+		const FString& TempDir = FMetalCompilerToolchain::Get()->GetLocalTempDir();
+		const FString MetalLibFilePath = FPaths::CreateTempFilename(*TempDir, TEXT("MSCLibrary-"), *FMetalCompilerToolchain::MetalLibraryExtension);
+		
+		ON_SCOPE_EXIT
+		{
+			// Clean up the .metallib file
+			IFileManager::Get().Delete(*MetalLibFilePath);
+		};
+		
+		// Serialize it back to the OutputFile
+		if (!FFileHelper::SaveArrayToFile(InMetalBytecode.OutputFile, *MetalLibFilePath))
+		{
+			Output.AddErrorf(TEXT("Error: MetalShaderConverter failed to write a metallib '%s'!"), *MetalLibFilePath);
+			Output.bSucceeded = false;
+
+			return;
+		}
+		
+		if (!FMetalCompilerToolchain::Get()->StripMetalLib(AppleSDKMac, MetalLibFilePath, MetalLibFilePath))
+		{
+			Output.AddErrorf(TEXT("Error: MetalShaderConverter failed to strip a metallib '%s'!"), *MetalLibFilePath);
+			Output.bSucceeded = false;
+
+			return;
+		}
+		
+		// Serialize it back to the OutputFile
+		if (!FFileHelper::LoadFileToArray(InMetalBytecode.OutputFile, *MetalLibFilePath))
+		{
+			Output.AddErrorf(TEXT("Error: MetalShaderConverter failed to read a metallib '%s'!"), *MetalLibFilePath);
+			Output.bSucceeded = false;
+
+			return;
+		}
+	}
 	
 	Output.bSucceeded = true;
-	
-	IRMetalLibBinaryDestroy(GeneratedMetalLib);
-	IRObjectDestroy(AirBytecode);
 };
 
 void FMetalCompileShaderMSC::DoCompileMetalShader(
diff --git a/Engine/Source/Developer/Apple/MetalShaderFormat/Private/MetalShaderCompiler.cpp b/Engine/Source/Developer/Apple/MetalShaderFormat/Private/MetalShaderCompiler.cpp
--- a/Engine/Source/Developer/Apple/MetalShaderFormat/Private/MetalShaderCompiler.cpp
+++ b/Engine/Source/Developer/Apple/MetalShaderFormat/Private/MetalShaderCompiler.cpp
@@ -57,12 +57,15 @@ THIRD_PARTY_INCLUDES_END
 #endif
 #endif
 
-static bool CompileProcessAllowsRuntimeShaderCompiling(const FShaderCompilerInput& InputCompilerEnvironment)
+static bool ShouldUseRuntimeShaderCompiling(const FShaderCompilerInput& InputCompilerEnvironment)
 {
-	bool bArchiving = InputCompilerEnvironment.Environment.CompilerFlags.Contains(CFLAG_Archive);
-	bool bDebug = InputCompilerEnvironment.Environment.CompilerFlags.Contains(CFLAG_Debug);
+	const bool bDebug = InputCompilerEnvironment.Environment.CompilerFlags.Contains(CFLAG_Debug);
+	return bDebug && InputCompilerEnvironment.Environment.GetCompileArgument(TEXT("METAL_ALLOW_RUNTIME_SHADER_COMPILING"), false);
+}
 
-	return !bArchiving && bDebug;
+bool ShouldStripIndividualMetalLib(const FShaderCompilerInput& InputCompilerEnvironment)
+{	
+	return InputCompilerEnvironment.Environment.GetCompileArgument(TEXT("METAL_STRIP_INDIVIDUAL_METALLIBS"), false);
 }
 
 constexpr uint16 GMetalMaxUniformBufferSlots = 32;
@@ -286,8 +289,7 @@ void BuildMetalShaderOutput(
 	const EShaderFrequency Frequency = ShaderOutput.Target.GetFrequency();
 	const bool bBindlessEnabled = ShaderInput.IsBindlessEnabled();
 
-	//TODO read from toolchain
-	const bool bIsMobile = (ShaderInput.Target.Platform == SP_METAL_ES3_1_IOS || ShaderInput.Target.Platform == SP_METAL_SM5_IOS || ShaderInput.Target.Platform == SP_METAL_ES3_1_TVOS || ShaderInput.Target.Platform == SP_METAL_SM5_TVOS || ShaderInput.Target.Platform == SP_METAL_SIM);
+	const bool bIsMobile = FMetalCompilerToolchain::Get()->IsMobile((EShaderPlatform)ShaderInput.Target.Platform);
 	bool bNoFastMath = ShaderInput.Environment.CompilerFlags.Contains(CFLAG_NoFastMath);
 	const bool bUsingWPO = ShaderInput.Environment.GetCompileArgument(TEXT("USES_WORLD_POSITION_OFFSET"), false);
 	if (bUsingWPO && (ShaderInput.Target.Platform == SP_METAL_SM5_IOS || ShaderInput.Target.Platform == SP_METAL_SM5_TVOS) && Frequency == SF_Vertex)
@@ -418,9 +420,6 @@ void BuildMetalShaderOutput(
 		uint16& Size = PackedGlobalArraySize.FindOrAdd(PackedGlobal.PackedType);
 		Size = FMath::Max<uint16>(BytesPerComponent * (PackedGlobal.Offset + PackedGlobal.Count), Size);
 	}
-
-	bool bUseMetalShaderConverter = false;
-	bUseMetalShaderConverter = bBindlessEnabled;
 	
 	// Packed Uniform Buffers
 	TMap<int, TMap<CrossCompiler::EPackedTypeName, uint16> > PackedUniformBuffersSize;
@@ -433,7 +432,7 @@ void BuildMetalShaderOutput(
 			// We need to distinguish Root/Global CBs when ShaderConverter is used (mainly for RT support);
 			// therefore we perform CB reflection during the previous compilation stage of the pipeline (and keep
 			// the vanilla path for SPIRV-Cross).
-			if (!bUseMetalShaderConverter)
+			if (!bBindlessEnabled)
 			{
 				HandleReflectedGlobalConstantBufferMember(
 					ShaderInput,
@@ -523,7 +522,7 @@ void BuildMetalShaderOutput(
 		HandleReflectedShaderSampler(SamplerState.Name, SamplerState.Index, SamplerMap[SamplerState.Name], ShaderOutput);
 	}
 
-    if (bUseMetalShaderConverter)
+    if (bBindlessEnabled)
     {
 #if PLATFORM_MAC || PLATFORM_WINDOWS
 
@@ -642,7 +641,7 @@ void BuildMetalShaderOutput(
 
 	FString MetalCode = FString(USFSource);
 
-	if (ShaderInput.Environment.CompilerFlags.Contains(CFLAG_Debug) && !bUseMetalShaderConverter)
+	if (ShaderInput.Environment.CompilerFlags.Contains(CFLAG_Debug) && !bBindlessEnabled)
 	{
 		MetalCode.InsertAt(0, FString::Printf(TEXT("// %s\n"), *CCHeader.Name));
 	}
@@ -663,8 +662,7 @@ void BuildMetalShaderOutput(
 			FString::Printf(TEXT("shader uses %d (%d) samplers exceeding the limit of %d\nSamplers:\n%s"),
 				Header.Bindings.NumSamplers, CCHeader.SamplerStates.Num(), MaxMetalSamplers, *SamplerList);
 	}
-	// TODO read from toolchain? this check isn't really doing exactly what it says
-	else if(CompileProcessAllowsRuntimeShaderCompiling(ShaderInput) && !bUseMetalShaderConverter)
+	else if(ShouldUseRuntimeShaderCompiling(ShaderInput) && !bBindlessEnabled)
 	{
 		// Write out the header and shader source code.
 		FMemoryWriter Ar(ShaderOutput.ShaderCode.GetWriteAccess(), true);
@@ -674,21 +672,23 @@ void BuildMetalShaderOutput(
 		Ar.Serialize((void*)USFSource, SourceLen + 1 - (USFSource - InShaderSource));
 		
 		ShaderOutput.ModifiedShaderSource = MetalCode;
-		
-		if(!bBindlessEnabled)
-		{
-			ShaderOutput.NumInstructions = NumLines;
-		}
+		ShaderOutput.NumInstructions = NumLines;
 		ShaderOutput.NumTextureSamplers = Header.Bindings.NumSamplers;
 		ShaderOutput.bSucceeded = true;
 	}
 	else
 	{
-		// TODO technically should probably check the version of the metal compiler to make sure it's recent enough to support -MO.
 		FString DebugInfo = TEXT("");
-		if (ShaderInput.Environment.CompilerFlags.Contains(CFLAG_ExtraShaderData))
+		const bool bGenerateSymbols = (ShaderInput.Environment.CompilerFlags.Contains(CFLAG_GenerateSymbols));
+		if (ShaderInput.Environment.CompilerFlags.Contains(CFLAG_ExtraShaderData) || bGenerateSymbols)
 		{
+			// Include line number and disable optimization
 			DebugInfo = TEXT("-gline-tables-only -MO");
+			if (bGenerateSymbols)
+			{
+				// Include source inside the air/metallib
+				DebugInfo += TEXT(" -frecord-sources");
+			}
 		}
 		
 		FString MathMode = bNoFastMath ? TEXT("-fno-fast-math") : TEXT("-ffast-math");
@@ -705,10 +705,8 @@ void BuildMetalShaderOutput(
 		EShaderPlatform ShaderPlatform = EShaderPlatform(ShaderInput.Target.Platform);
 		const bool bMetalCompilerAvailable = FMetalCompilerToolchain::Get()->IsCompilerAvailable();
 		
-		bool bDebugInfoSucceded = false;
 		FMetalShaderBytecode Bytecode;
 		FMetalShaderBytecode ExtendedBytecode;
-		FMetalShaderDebugInfo DebugCode;
 		
 		FString HashedName = FString::Printf(TEXT("%u_%u"), SourceCRCLen, SourceCRC);
 		
@@ -733,7 +731,7 @@ void BuildMetalShaderOutput(
 			Error->ErrorLineString = TEXT("0");
 			Error->StrippedErrorMessage = FString(Message);
 		}
-		else if (bUseMetalShaderConverter)
+		else if (bBindlessEnabled)
         {
             // The base name (which is <temp>/CRCHash_Length)
             FString BaseFileName = FPaths::Combine(TempDir, HashedName);
@@ -808,8 +806,9 @@ void BuildMetalShaderOutput(
 			Job.Standard = Standard;
 			Job.SourceCRCLen = SourceCRCLen;
 			Job.SourceCRC = SourceCRC;
-			Job.bRetainObjectFile = ShaderInput.Environment.CompilerFlags.Contains(CFLAG_Archive);
+			Job.bUseNativeShaderLibrary = ShaderInput.Environment.CompilerFlags.Contains(CFLAG_Archive);
 			Job.bOptimizeForSize = ShaderInput.Environment.GetCompileArgument(TEXT("METAL_OPTIMIZE_FOR_SIZE"), false);
+			Job.bStripIndividualMetalLibs = ShouldStripIndividualMetalLib(ShaderInput);
 			Job.bCompileAsPCH = false;
 			Job.ReturnCode = 0;
 			
@@ -820,24 +819,7 @@ void BuildMetalShaderOutput(
 			Job.ModuleCacheDirectory = ModulesCachePath;
 			
 			bSucceeded = FMetalCompilerToolchain::Get()->CompileMetalShader(Job, Bytecode);
-			if (bSucceeded)
-			{
-				if (!bIsMobile && !ShaderInput.Environment.CompilerFlags.Contains(CFLAG_Archive))
-				{
-					uint32 CodeSize = FCStringAnsi::Strlen(TCHAR_TO_UTF8(*MetalCode)) + 1;
-
-					int32 CompressedSize = FCompression::CompressMemoryBound(NAME_Zlib, CodeSize);
-					DebugCode.CompressedData.SetNum(CompressedSize);
-
-					if (FCompression::CompressMemory(NAME_Zlib, DebugCode.CompressedData.GetData(), CompressedSize, TCHAR_TO_UTF8(*MetalCode), CodeSize))
-					{
-						DebugCode.UncompressedSize = CodeSize;
-						DebugCode.CompressedData.SetNum(CompressedSize);
-						DebugCode.CompressedData.Shrink();
-					}
-				}
-			}
-			else
+			if (!bSucceeded)
 			{
 				FShaderCompilerError* Error = new(OutErrors) FShaderCompilerError();
 				Error->ErrorVirtualFilePath = MetalFileName;
@@ -866,25 +848,15 @@ void BuildMetalShaderOutput(
 			
 			if(ExtendedBytecode.OutputFile.Num())
 			{
-				ShaderOutput.ShaderCode.AddOptionalData(EShaderOptionalDataKey::VendorExtension, ExtendedBytecode.ObjectFile.GetData(), ExtendedBytecode.ObjectFile.Num());
-			}
-			
-			if (bDebugInfoSucceded && !ShaderInput.Environment.CompilerFlags.Contains(CFLAG_Archive) && DebugCode.CompressedData.Num())
-			{
-				ShaderOutput.ShaderCode.AddOptionalData(EShaderOptionalDataKey::CompressedDebugCode, DebugCode.CompressedData.GetData(), DebugCode.CompressedData.Num());
-				ShaderOutput.ShaderCode.AddOptionalData(EShaderOptionalDataKey::NativePath, TCHAR_TO_UTF8(*Bytecode.NativePath));
-				ShaderOutput.ShaderCode.AddOptionalData(EShaderOptionalDataKey::UncompressedSize, (const uint8*)&DebugCode.UncompressedSize, sizeof(DebugCode.UncompressedSize));
+				ShaderOutput.ShaderCode.AddOptionalData(EShaderOptionalDataKey::VendorExtension, ExtendedBytecode.OutputFile.GetData(), ExtendedBytecode.OutputFile.Num());
 			}
 			
 			if (ShaderInput.Environment.CompilerFlags.Contains(CFLAG_ExtraShaderData) && !ExtendedBytecode.OutputFile.Num())
 			{
 				// store data we can pickup later with ShaderCode.FindOptionalData(FShaderCodeName::Key), could be removed for shipping
 				ShaderOutput.ShaderCode.AddOptionalData(FShaderCodeName::Key, TCHAR_TO_UTF8(*ShaderInput.GenerateShaderName()));
-				if (DebugCode.CompressedData.Num() == 0)
-				{
-					ShaderOutput.ShaderCode.AddOptionalData(EShaderOptionalDataKey::SourceCode, TCHAR_TO_UTF8(*MetalCode));
-					ShaderOutput.ShaderCode.AddOptionalData(EShaderOptionalDataKey::NativePath, TCHAR_TO_UTF8(*Bytecode.NativePath));
-				}
+				ShaderOutput.ShaderCode.AddOptionalData(EShaderOptionalDataKey::SourceCode, TCHAR_TO_UTF8(*MetalCode));
+				ShaderOutput.ShaderCode.AddOptionalData(EShaderOptionalDataKey::NativePath, TCHAR_TO_UTF8(*Bytecode.NativePath));
 			}
 			
 			ShaderOutput.NumTextureSamplers = Header.Bindings.NumSamplers;
@@ -1542,18 +1514,7 @@ bool FinalizeLibrary_Metal(FName const& Format, FString const& WorkingDir, FStri
 					
 					if (bStripLibs)
 					{
-						UE_LOG(LogMetalShaderCompiler, Display, TEXT("Stripping lib %s Size: %d"), *BatchOutputPath, IFileManager::Get().FileSize(*BatchOutputPath));
-						
-						FString MetalStripParams = FString::Printf(TEXT("-S -T -o=\"%s\" \"%s\""), *BatchOutputPath, *BatchOutputPath);
-						bSuccess = Toolchain->ExecMetalStrip(SDK, *MetalStripParams, &ReturnCode, &Results, &Errors);
-						
-						// handle compile error
-						if (!bSuccess)
-						{
-							UE_LOG(LogShaders, Error, TEXT("Archiving failed: metal-pack failed with code %d: %s %s"), ReturnCode, *Results, *Errors);
-						}
-						
-						UE_LOG(LogMetalShaderCompiler, Display, TEXT("Stripped size of %s Size: %d"), *BatchOutputPath, IFileManager::Get().FileSize(*BatchOutputPath));
+						bSuccess = Toolchain->StripMetalLib(SDK, BatchOutputPath, BatchOutputPath);
 					}
 				}
 
@@ -1691,32 +1652,9 @@ bool FinalizeLibrary_Metal(FName const& Format, FString const& WorkingDir, FStri
 				return false;
 			}
 			
-			if(bStripLibs)
+			if (bStripLibs)
 			{
-				// Archive build phase - like unix ar, build metal archive from all the object files
-				UE_LOG(LogMetalShaderCompiler, Display, TEXT("Stripping lib %s Size: %d"), *LibraryPath, IFileManager::Get().FileSize(*LibraryPath));
-				
-                FString MetalStripParams = FString::Printf(TEXT("-S -T -o=\"%s\" \"%s\""), *LibraryPath, *LibraryPath);
-				
-				bSuccess = Toolchain->ExecMetalStrip(SDK, *MetalStripParams, &ReturnCode, &Results, &Errors);
-				
-				if (!bSuccess || ReturnCode != 0)
-				{
-					UE_LOG(LogShaders, Error, TEXT("Stripping failed: metallib failed with code %d: %s %s"), ReturnCode, *Results, *Errors);
-					return false;
-				}
-				
-				UE_LOG(LogMetalShaderCompiler, Display, TEXT("Stripped size of %s Size: %d"), *LibraryPath, IFileManager::Get().FileSize(*LibraryPath));
-				
-				check(LocalMetalLibPath == LibraryPath);
-				
-				bOK = (IFileManager::Get().FileSize(*LibraryPath) > 0);
-				
-				if (!bOK)
-				{
-					UE_LOG(LogShaders, Error, TEXT("Stripping failed: failed to copy to local destination: %s"), *LibraryPath);
-					return bOK;
-				}
+				bOK = Toolchain->StripMetalLib(SDK, LibraryPath, LibraryPath);
 			}
 		}
 	}
diff --git a/Engine/Source/Developer/Apple/MetalShaderFormat/Private/MetalShaderCompiler.h b/Engine/Source/Developer/Apple/MetalShaderFormat/Private/MetalShaderCompiler.h
--- a/Engine/Source/Developer/Apple/MetalShaderFormat/Private/MetalShaderCompiler.h
+++ b/Engine/Source/Developer/Apple/MetalShaderFormat/Private/MetalShaderCompiler.h
@@ -10,33 +10,6 @@
 
 #include <string>
 
-struct FMetalShaderDebugInfoJob
-{
-	FName ShaderFormat;
-	FSHAHash Hash;
-	FString CompilerVersion;
-	FString MinOSVersion;
-	FString DebugInfo;
-	FString MathMode;
-	FString Standard;
-	uint32 SourceCRCLen;
-	uint32 SourceCRC;
-	
-	FString MetalCode;
-};
-
-struct FMetalShaderDebugInfo
-{
-	uint32 UncompressedSize;
-	TArray<uint8> CompressedData;
-	
-	friend FArchive& operator<<( FArchive& Ar, FMetalShaderDebugInfo& Info )
-	{
-		Ar << Info.UncompressedSize << Info.CompressedData;
-		return Ar;
-	}
-};
-
 struct FMetalShaderBytecodeJob
 {
 	FName ShaderFormat;
@@ -56,9 +29,10 @@ struct FMetalShaderBytecodeJob
 	FString ModuleCacheDirectory;
 	uint32 SourceCRCLen;
 	uint32 SourceCRC;
-	bool bRetainObjectFile;
+	bool bUseNativeShaderLibrary;
 	bool bCompileAsPCH;
 	bool bOptimizeForSize;
+	bool bStripIndividualMetalLibs;
 	
 	FString Message;
 	FString Results;
diff --git a/Engine/Source/Developer/Apple/MetalShaderFormat/Private/MetalShaderFormat.cpp b/Engine/Source/Developer/Apple/MetalShaderFormat/Private/MetalShaderFormat.cpp
--- a/Engine/Source/Developer/Apple/MetalShaderFormat/Private/MetalShaderFormat.cpp
+++ b/Engine/Source/Developer/Apple/MetalShaderFormat/Private/MetalShaderFormat.cpp
@@ -1289,9 +1289,9 @@ bool FMetalCompilerToolchain::CompileMetalShader(FMetalShaderBytecodeJob& Job, F
 	}
 
 	// At this point we have an .air file and a .metallib file
-	if (Job.bRetainObjectFile)
+	if (Job.bUseNativeShaderLibrary)
 	{
-		// Retain the .air. This usually means we are using shared native libraries.
+		// Retain the .air.
 		bool bSuccess = FFileHelper::LoadFileToArray(Output.ObjectFile, *LocalOutputMetalAIRFilePath);
 		
 		// Clean up the .air file
@@ -1307,6 +1307,16 @@ bool FMetalCompilerToolchain::CompileMetalShader(FMetalShaderBytecodeJob& Job, F
 	{
 		// Clean up the .air file
 		IFileManager::Get().Delete(*LocalOutputMetalAIRFilePath);
+		
+		// For non-native shader library we need to strip the metallib now
+		if (Job.bStripIndividualMetalLibs)
+		{
+			if (!StripMetalLib(SDK, LocalOutputMetalLibFilePath, LocalOutputMetalLibFilePath))
+			{
+				Job.Message = FString::Printf(TEXT("Failed to strip metallib %s"), *LocalOutputMetalLibFilePath);
+				return false;
+			}
+		}
 	}
 
 	{
@@ -1326,3 +1336,31 @@ bool FMetalCompilerToolchain::CompileMetalShader(FMetalShaderBytecodeJob& Job, F
 
 	return true;
 }
+
+bool FMetalCompilerToolchain::StripMetalLib(EAppleSDKType SDK, const FString& InMetalLibFilePath, const FString& OutMetalLibFilePath) const
+{
+	UE_LOG(LogMetalShaderCompiler, Display, TEXT("Stripping lib %s Size: %lld"), *InMetalLibFilePath, IFileManager::Get().FileSize(*InMetalLibFilePath));
+	
+	// Strip
+	FString MetalStripParams = FString::Printf(TEXT("-S -T -o=\"%s\" \"%s\""), *OutMetalLibFilePath, *InMetalLibFilePath);
+	
+	int32 ReturnCode = 0;
+	FString Results;
+	FString Errors;
+	bool bSuccess = this->ExecMetalStrip(SDK, *MetalStripParams, &ReturnCode, &Results, &Errors);
+	
+	if (!bSuccess || ReturnCode != 0)
+	{
+		UE_LOG(LogShaders, Error, TEXT("Stripping failed: metallib failed with code %d: %s %s"), ReturnCode, *Results, *Errors);
+		return false;
+	}
+	
+	int64 FileSize = IFileManager::Get().FileSize(*OutMetalLibFilePath);
+	if (FileSize <= 0)
+	{
+		UE_LOG(LogShaders, Error, TEXT("Stripping failed: metallib size == %lld: %s"), FileSize, *OutMetalLibFilePath);
+		return false;
+	}
+	UE_LOG(LogMetalShaderCompiler, Display, TEXT("Stripped size of %s Size: %lld"), *OutMetalLibFilePath, FileSize);
+	return true;
+}
diff --git a/Engine/Source/Developer/Apple/MetalShaderFormat/Private/MetalShaderFormat.h b/Engine/Source/Developer/Apple/MetalShaderFormat/Private/MetalShaderFormat.h
--- a/Engine/Source/Developer/Apple/MetalShaderFormat/Private/MetalShaderFormat.h
+++ b/Engine/Source/Developer/Apple/MetalShaderFormat/Private/MetalShaderFormat.h
@@ -66,6 +66,8 @@ public:
 	// Takes a Job and compiles a shader. Produces .air and .metallib
 	bool CompileMetalShader(struct FMetalShaderBytecodeJob& Job, struct FMetalShaderBytecode& Output) const;
 	
+	bool StripMetalLib(EAppleSDKType SDK, const FString& InMetalLibFilePath, const FString& OutMetalLibFilePath) const;
+	
 	// Executes 'Command' on the local machine
 	bool ExecGenericCommand(const TCHAR* Command, const TCHAR* Params, int32* OutReturnCode, FString* OutStdOut, FString* OutStdErr, bool bIsConsoleApp = false) const;
 	// Executes the metal frontend compiler for 'SDK' on the local or remote machine, depending on the current configuration
diff --git a/Engine/Source/Runtime/Engine/Private/ShaderCompiler/ShaderCompiler.cpp b/Engine/Source/Runtime/Engine/Private/ShaderCompiler/ShaderCompiler.cpp
--- a/Engine/Source/Runtime/Engine/Private/ShaderCompiler/ShaderCompiler.cpp
+++ b/Engine/Source/Runtime/Engine/Private/ShaderCompiler/ShaderCompiler.cpp
@@ -21,7 +21,9 @@
 #include "Interfaces/ITargetPlatform.h"
 #include "Interfaces/ITargetPlatformManagerModule.h"
 #include "Internationalization/LocKeyFuncs.h"
+#include "IOS/IOSPlatformProperties.h"
 #include "Logging/StructuredLog.h"
+#include "Mac/MacPlatformProperties.h"
 #include "MaterialShared.h"
 #include "Materials/Material.h"
 #include "Materials/MaterialInstance.h"
@@ -3744,20 +3746,11 @@ void GlobalBeginCompileShader(
 		{
 			Input.Environment.CompilerFlags.Add(CFLAG_BoundsChecking);
 		}
-		
-		// Check whether we can compile metal shaders to bytecode - avoids poisoning the DDC
-		static ITargetPlatformManagerModule& TPM = GetTargetPlatformManagerRef();
-		const IShaderFormat* Compiler = TPM.FindShaderFormat(ShaderFormatName);
-		static const bool bCanCompileOfflineMetalShaders = Compiler && Compiler->CanCompileBinaryShaders();
-		if (!bCanCompileOfflineMetalShaders)
-		{
-			Input.Environment.CompilerFlags.Add(CFLAG_Debug);
-		}
-		
-		// Shaders built for archiving - for Metal that requires compiling the code in a different way so that we can strip it later
+
+		// Decide wether to use Metal Native Shader Library or not
 		bool bArchive = false;
 		GConfig->GetBool(TEXT("/Script/UnrealEd.ProjectPackagingSettings"), TEXT("bSharedMaterialNativeLibraries"), bArchive, GGameIni);
-		if (bCanCompileOfflineMetalShaders && bArchive)
+		if (bArchive)
 		{
 			Input.Environment.CompilerFlags.Add(CFLAG_Archive);
 		}
@@ -3770,23 +3763,26 @@ void GlobalBeginCompileShader(
 			bool bForceFloats = false;
 			bool bEnableMathOptimisations = true;
 			bool bMetalOptimizeForSize = false;
-
-			if (IsPCPlatform(EShaderPlatform(Target.Platform)))
+			bool bAllowsRuntimeShaderCompiling = false;
+			bool bStripIndividualMetalLibs = false;
+			
+			const bool bIsMac = IsPCPlatform(EShaderPlatform(Target.Platform));
+			const TCHAR* RuntimeSettingsClassName = bIsMac ? FMacPlatformProperties<false, false, false>::GetRuntimeSettingsClassName() : FIOSPlatformProperties::GetRuntimeSettingsClassName();
+
+			GConfig->GetBool(RuntimeSettingsClassName, TEXT("UseFastIntrinsics"), bAllowFastIntrinsics, GEngineIni);
+			GConfig->GetBool(RuntimeSettingsClassName, TEXT("EnableMathOptimisations"), bEnableMathOptimisations, GEngineIni);
+			GConfig->GetBool(RuntimeSettingsClassName, TEXT("ForceFloats"), bForceFloats, GEngineIni);
+			GConfig->GetBool(RuntimeSettingsClassName, TEXT("MetalOptimizeForSize"), bMetalOptimizeForSize, GEngineIni);
+			GConfig->GetBool(RuntimeSettingsClassName, TEXT("AllowsRuntimeShaderCompiling"), bAllowsRuntimeShaderCompiling, GEngineIni);
+			GConfig->GetBool(RuntimeSettingsClassName, TEXT("StripIndividualMetalLibs"), bStripIndividualMetalLibs, GEngineIni);
+			
+			if (bIsMac)
 			{
-				GConfig->GetBool(TEXT("/Script/MacTargetPlatform.MacTargetSettings"), TEXT("UseFastIntrinsics"), bAllowFastIntrinsics, GEngineIni);
-				GConfig->GetBool(TEXT("/Script/MacTargetPlatform.MacTargetSettings"), TEXT("EnableMathOptimisations"), bEnableMathOptimisations, GEngineIni);
-				GConfig->GetBool(TEXT("/Script/MacTargetPlatform.MacTargetSettings"), TEXT("MetalOptimizeForSize"), bMetalOptimizeForSize, GEngineIni);
-                
-                // No half precision support on MacOS at the moment
-                bForceFloats = true;
+				// No half precision support on MacOS at the moment
+				bForceFloats = true;
 			}
 			else
 			{
-				GConfig->GetBool(TEXT("/Script/IOSRuntimeSettings.IOSRuntimeSettings"), TEXT("UseFastIntrinsics"), bAllowFastIntrinsics, GEngineIni);
-				GConfig->GetBool(TEXT("/Script/IOSRuntimeSettings.IOSRuntimeSettings"), TEXT("EnableMathOptimisations"), bEnableMathOptimisations, GEngineIni);
-				GConfig->GetBool(TEXT("/Script/IOSRuntimeSettings.IOSRuntimeSettings"), TEXT("ForceFloats"), bForceFloats, GEngineIni);
-				GConfig->GetBool(TEXT("/Script/IOSRuntimeSettings.IOSRuntimeSettings"), TEXT("MetalOptimizeForSize"), bMetalOptimizeForSize, GEngineIni);
-                
 				// Force no development shaders on iOS
 				bAllowDevelopmentShaderCompile = false;
 			}
@@ -3796,6 +3792,26 @@ void GlobalBeginCompileShader(
 			Input.Environment.SetCompileArgument(TEXT("METAL_USE_FAST_INTRINSICS"), bAllowFastIntrinsics);
             Input.Environment.SetCompileArgument(TEXT("METAL_OPTIMIZE_FOR_SIZE"), bMetalOptimizeForSize);
 			
+			static const auto CVarSymbols = IConsoleManager::Get().FindConsoleVariable(TEXT("r.Shaders.Symbols"));
+			static const auto CVarExtraData = IConsoleManager::Get().FindConsoleVariable(TEXT("r.Shaders.ExtraData"));
+			
+			if (CVarSymbols && CVarExtraData)
+			{
+				bStripIndividualMetalLibs = bStripIndividualMetalLibs && CVarSymbols->GetInt() == 0 && CVarExtraData->GetInt() == 0;
+			}
+			Input.Environment.SetCompileArgument(TEXT("METAL_STRIP_INDIVIDUAL_METALLIBS"), bStripIndividualMetalLibs);
+			
+			
+			// Check whether we can compile metal shaders to bytecode - avoids poisoning the DDC
+			static ITargetPlatformManagerModule& TPM = GetTargetPlatformManagerRef();
+			const IShaderFormat* Compiler = TPM.FindShaderFormat(ShaderFormatName);
+			const bool bCanCompileOfflineMetalShaders = Compiler && Compiler->CanCompileBinaryShaders();
+			if (!bCanCompileOfflineMetalShaders)
+			{
+				Input.Environment.CompilerFlags.Add(CFLAG_Debug);
+			}
+			Input.Environment.SetCompileArgument(TEXT("METAL_ALLOW_RUNTIME_SHADER_COMPILING"), !bCanCompileOfflineMetalShaders || bAllowsRuntimeShaderCompiling);
+			
 			// Same as console-variable above, but that's global and this is per-platform, per-project
 			if (!bEnableMathOptimisations)
 			{
diff --git a/Engine/Source/Runtime/RenderCore/Private/Shader.cpp b/Engine/Source/Runtime/RenderCore/Private/Shader.cpp
--- a/Engine/Source/Runtime/RenderCore/Private/Shader.cpp
+++ b/Engine/Source/Runtime/RenderCore/Private/Shader.cpp
@@ -13,12 +13,14 @@
 #include "Interfaces/ITargetPlatformManagerModule.h"
 #include "Interfaces/IShaderFormat.h"
 #include "Internationalization/Regex.h"
+#include "IOS/IOSPlatformProperties.h"
 #include "Serialization/ShaderKeyGenerator.h"
 #include "ShaderCodeLibrary.h"
 #include "ShaderCore.h"
 #include "ShaderCompilerCore.h"
 #include "RenderUtils.h"
 #include "StereoRenderUtils.h"
+#include "Mac/MacPlatformProperties.h"
 #include "Misc/ConfigCacheIni.h"
 #include "Misc/ScopeLock.h"
 #include "UObject/RenderingObjectVersion.h"
@@ -2204,19 +2206,23 @@ void ShaderMapAppendKey(EShaderPlatform Platform, FShaderKeyGenerator& KeyGen)
 		bool bEnableMathOptimisations = true;
 		bool bForceFloats = false;
 		bool bMetalOptimizeBySize = false;
-
-		if (IsPCPlatform(Platform))
-		{
-			GConfig->GetBool(TEXT("/Script/MacTargetPlatform.MacTargetSettings"), TEXT("UseFastIntrinsics"), bAllowFastIntrinsics, GEngineIni);
-			GConfig->GetBool(TEXT("/Script/MacTargetPlatform.MacTargetSettings"), TEXT("EnableMathOptimisations"), bEnableMathOptimisations, GEngineIni);
-			GConfig->GetBool(TEXT("/Script/MacTargetPlatform.MacTargetSettings"), TEXT("MetalOptimizeBySize"), bMetalOptimizeBySize, GEngineIni);
-		}
-		else
+		bool bAllowsRuntimeShaderCompiling = false;
+		bool bStripIndividualMetalLibs = false;
+		
+		const bool bIsMac = IsPCPlatform(Platform);
+		const TCHAR* RuntimeSettingsClassName = bIsMac ? FMacPlatformProperties<false, false, false>::GetRuntimeSettingsClassName() : FIOSPlatformProperties::GetRuntimeSettingsClassName();
+
+		GConfig->GetBool(RuntimeSettingsClassName, TEXT("UseFastIntrinsics"), bAllowFastIntrinsics, GEngineIni);
+		GConfig->GetBool(RuntimeSettingsClassName, TEXT("EnableMathOptimisations"), bEnableMathOptimisations, GEngineIni);
+		GConfig->GetBool(RuntimeSettingsClassName, TEXT("ForceFloats"), bForceFloats, GEngineIni);
+		GConfig->GetBool(RuntimeSettingsClassName, TEXT("MetalOptimizeBySize"), bMetalOptimizeBySize, GEngineIni);
+		GConfig->GetBool(RuntimeSettingsClassName, TEXT("AllowsRuntimeShaderCompiling"), bAllowsRuntimeShaderCompiling, GEngineIni);
+		GConfig->GetBool(RuntimeSettingsClassName, TEXT("StripIndividualMetalLibs"), bStripIndividualMetalLibs, GEngineIni);
+		
+		if (bIsMac)
 		{
-			GConfig->GetBool(TEXT("/Script/IOSRuntimeSettings.IOSRuntimeSettings"), TEXT("UseFastIntrinsics"), bAllowFastIntrinsics, GEngineIni);
-			GConfig->GetBool(TEXT("/Script/IOSRuntimeSettings.IOSRuntimeSettings"), TEXT("EnableMathOptimisations"), bEnableMathOptimisations, GEngineIni);
-			GConfig->GetBool(TEXT("/Script/IOSRuntimeSettings.IOSRuntimeSettings"), TEXT("ForceFloats"), bForceFloats, GEngineIni);
-			GConfig->GetBool(TEXT("/Script/IOSRuntimeSettings.IOSRuntimeSettings"), TEXT("MetalOptimizeBySize"), bMetalOptimizeBySize, GEngineIni);
+			// No half precision support on MacOS at the moment
+			bForceFloats = true;
 		}
 		
 		if (bAllowFastIntrinsics)
@@ -2243,6 +2249,18 @@ void ShaderMapAppendKey(EShaderPlatform Platform, FShaderKeyGenerator& KeyGen)
 			KeyGen.AppendSeparator();
 			KeyGen.Append(TEXT("Os"));
         }
+		
+		if(bAllowsRuntimeShaderCompiling)
+		{
+			KeyGen.AppendSeparator();
+			KeyGen.Append(TEXT("AllowsRuntimeShaderCompiling"));
+		}
+		
+		if(bStripIndividualMetalLibs)
+		{
+			KeyGen.AppendSeparator();
+			KeyGen.Append(TEXT("StripIndividualMetalLibs"));
+		}
         
 		KeyGen.AppendSeparator();
 		KeyGen.Append(TEXT("IAB"));
