Metal RHI

- Optionally disable stripping of types for Metal libraries as this can cause errors with shaders build via MSC

Submitted on behalf of Carl.Lloyd

#tests local cook & preflight
#Synced-CL 51095579

diff --git a/Engine/Source/Developer/Apple/MetalShaderFormat/Private/MetalCompileShaderMSC.cpp b/Engine/Source/Developer/Apple/MetalShaderFormat/Private/MetalCompileShaderMSC.cpp
--- a/Engine/Source/Developer/Apple/MetalShaderFormat/Private/MetalCompileShaderMSC.cpp
+++ b/Engine/Source/Developer/Apple/MetalShaderFormat/Private/MetalCompileShaderMSC.cpp
@@ -897,7 +897,7 @@ void CompileToMetalLib(FMetalShaderBytecode& InMetalBytecode,
 			return;
 		}
 		
-		if (!FMetalCompilerToolchain::Get()->StripMetalLib(AppleSDKMac, MetalLibFilePath, MetalLibFilePath))
+		if (!FMetalCompilerToolchain::Get()->StripMetalLib(AppleSDKMac, MetalLibFilePath, MetalLibFilePath, false))
 		{
 			Output.AddErrorf(TEXT("Error: MetalShaderConverter failed to strip a metallib '%s'!"), *MetalLibFilePath);
 			Output.bSucceeded = false;
diff --git a/Engine/Source/Developer/Apple/MetalShaderFormat/Private/MetalShaderCompiler.cpp b/Engine/Source/Developer/Apple/MetalShaderFormat/Private/MetalShaderCompiler.cpp
--- a/Engine/Source/Developer/Apple/MetalShaderFormat/Private/MetalShaderCompiler.cpp
+++ b/Engine/Source/Developer/Apple/MetalShaderFormat/Private/MetalShaderCompiler.cpp
@@ -1515,7 +1515,7 @@ bool FinalizeLibrary_Metal(FName const& Format, FString const& WorkingDir, FStri
 					
 					if (bStripLibs)
 					{
-						bSuccess = Toolchain->StripMetalLib(SDK, BatchOutputPath, BatchOutputPath);
+						bSuccess = Toolchain->StripMetalLib(SDK, BatchOutputPath, BatchOutputPath, false);
 					}
 				}
 
@@ -1655,7 +1655,7 @@ bool FinalizeLibrary_Metal(FName const& Format, FString const& WorkingDir, FStri
 			
 			if (bStripLibs)
 			{
-				bOK = Toolchain->StripMetalLib(SDK, LibraryPath, LibraryPath);
+				bOK = Toolchain->StripMetalLib(SDK, LibraryPath, LibraryPath, true);
 			}
 		}
 	}
diff --git a/Engine/Source/Developer/Apple/MetalShaderFormat/Private/MetalShaderFormat.cpp b/Engine/Source/Developer/Apple/MetalShaderFormat/Private/MetalShaderFormat.cpp
--- a/Engine/Source/Developer/Apple/MetalShaderFormat/Private/MetalShaderFormat.cpp
+++ b/Engine/Source/Developer/Apple/MetalShaderFormat/Private/MetalShaderFormat.cpp
@@ -1312,7 +1312,7 @@ bool FMetalCompilerToolchain::CompileMetalShader(FMetalShaderBytecodeJob& Job, F
 		// For non-native shader library we need to strip the metallib now
 		if (Job.bStripIndividualMetalLibs)
 		{
-			if (!StripMetalLib(SDK, LocalOutputMetalLibFilePath, LocalOutputMetalLibFilePath))
+			if (!StripMetalLib(SDK, LocalOutputMetalLibFilePath, LocalOutputMetalLibFilePath, true))
 			{
 				Job.Message = FString::Printf(TEXT("Failed to strip metallib %s"), *LocalOutputMetalLibFilePath);
 				return false;
@@ -1338,12 +1338,14 @@ bool FMetalCompilerToolchain::CompileMetalShader(FMetalShaderBytecodeJob& Job, F
 	return true;
 }
 
-bool FMetalCompilerToolchain::StripMetalLib(EAppleSDKType SDK, const FString& InMetalLibFilePath, const FString& OutMetalLibFilePath) const
+bool FMetalCompilerToolchain::StripMetalLib(EAppleSDKType SDK, const FString& InMetalLibFilePath, const FString& OutMetalLibFilePath, const bool bStripTypes) const
 {
 	UE_LOG(LogMetalShaderCompiler, Display, TEXT("Stripping lib %s Size: %lld"), *InMetalLibFilePath, IFileManager::Get().FileSize(*InMetalLibFilePath));
 	
+	FString StripTypesArg = bStripTypes ? FString(TEXT("-T")) : FString(); 
+	
 	// Strip
-	FString MetalStripParams = FString::Printf(TEXT("-S -T -o=\"%s\" \"%s\""), *OutMetalLibFilePath, *InMetalLibFilePath);
+	FString MetalStripParams = FString::Printf(TEXT("-S %s -o=\"%s\" \"%s\""), *StripTypesArg, *OutMetalLibFilePath, *InMetalLibFilePath);
 	
 	int32 ReturnCode = 0;
 	FString Results;
diff --git a/Engine/Source/Developer/Apple/MetalShaderFormat/Private/MetalShaderFormat.h b/Engine/Source/Developer/Apple/MetalShaderFormat/Private/MetalShaderFormat.h
--- a/Engine/Source/Developer/Apple/MetalShaderFormat/Private/MetalShaderFormat.h
+++ b/Engine/Source/Developer/Apple/MetalShaderFormat/Private/MetalShaderFormat.h
@@ -66,7 +66,8 @@ public:
 	// Takes a Job and compiles a shader. Produces .air and .metallib
 	bool CompileMetalShader(struct FMetalShaderBytecodeJob& Job, struct FMetalShaderBytecode& Output) const;
 	
-	bool StripMetalLib(EAppleSDKType SDK, const FString& InMetalLibFilePath, const FString& OutMetalLibFilePath) const;
+	// Strips the Metal lib of debug symbols and optionally type info
+	bool StripMetalLib(EAppleSDKType SDK, const FString& InMetalLibFilePath, const FString& OutMetalLibFilePath, const bool bStripTypes) const;
 	
 	// Executes 'Command' on the local machine
 	bool ExecGenericCommand(const TCHAR* Command, const TCHAR* Params, int32* OutReturnCode, FString* OutStdOut, FString* OutStdErr, bool bIsConsoleApp = false) const;
