Fixed mac es3.1 support for cooked builds

#Synced-CL 52243953

diff --git a/Engine/Source/Developer/Apple/MetalShaderFormat/Private/MetalShaderFormat.cpp b/Engine/Source/Developer/Apple/MetalShaderFormat/Private/MetalShaderFormat.cpp
--- a/Engine/Source/Developer/Apple/MetalShaderFormat/Private/MetalShaderFormat.cpp
+++ b/Engine/Source/Developer/Apple/MetalShaderFormat/Private/MetalShaderFormat.cpp
@@ -205,7 +205,7 @@ public:
 			NumShadersPerLibrary = CVarNumShadersPerLibraryBindlessWithSymbols.GetValueOnAnyThread();
 		}
 		
-		const FString ArchivePath = (WorkingDirectory / ShaderFormatAndShaderPlatformName.GetPlainNameString());
+		const FString ArchivePath = (WorkingDirectory / ShaderFormatAndShaderPlatformName.ToString());
 		IFileManager::Get().DeleteDirectory(*ArchivePath, false, true);
 		IFileManager::Get().MakeDirectory(*ArchivePath);
 
@@ -301,7 +301,8 @@ public:
 		}
 
 		bool bOK = false;
-		FString LibraryPlatformName = FString::Printf(TEXT("%s_%s"), *LibraryName, *ShaderFormatAndShaderPlatformName.GetPlainNameString());
+		FString LibraryPlatformName = FString::Printf(TEXT("%s_%s"), *LibraryName, *ShaderFormatAndShaderPlatformName.ToString());
+
 		LibraryPlatformName.ToLowerInline();
 		volatile int32 CompiledLibraries = 0;
 		FGraphEventArray Tasks;
@@ -361,7 +362,7 @@ public:
 					IFileManager::Get().FindFilesRecursive(FilesToArchive, *DebugOutputDir, TEXT("*.metal"), true, false, false);
 
 					//Write the local file names into the target file
-					const FString DebugDir = DebugOutputDir / *ShaderFormatAndShaderPlatformName.GetPlainNameString();
+					const FString DebugDir = DebugOutputDir / *ShaderFormatAndShaderPlatformName.ToString();
 
 					for (FString FileName : FilesToArchive)
 					{
